Porting parameter smoothing for reverb from 4.23-Release

diff --git a/Engine/Source/Runtime/SignalProcessing/Private/DynamicDelayAPF.cpp b/Engine/Source/Runtime/SignalProcessing/Private/DynamicDelayAPF.cpp
--- a/Engine/Source/Runtime/SignalProcessing/Private/DynamicDelayAPF.cpp
+++ b/Engine/Source/Runtime/SignalProcessing/Private/DynamicDelayAPF.cpp
@@ -5,14 +5,14 @@
 
 using namespace Audio;
 
-FDynamicDelayAPF::FDynamicDelayAPF(float InG, int32 InMinDelay, int32 InMaxDelay, int32 InMaxNumInternalBufferSamples)
+FDynamicDelayAPF::FDynamicDelayAPF(float InG, int32 InMinDelay, int32 InMaxDelay, int32 InMaxNumInternalBufferSamples, float InSampleRate)
 :	MinDelay(InMinDelay)
 ,	MaxDelay(InMaxDelay)
 ,	NumDelaySamples(InMinDelay - 1)
 ,	NumInternalBufferSamples(InMaxNumInternalBufferSamples)
 {
 	G.SetValueInterrupt(0.0f);
-	G.Init(48000); // TODO: use actual sample rate (will warp ease time slightly right now)
+	G.Init(InSampleRate);
 
 	checkf(NumDelaySamples >= 0, TEXT("Minimum delay must be atleast 1"));
 	// NumInternalBufferSamples must be less than the length of the delay to support buffer indexing logic.
diff --git a/Engine/Source/Runtime/SignalProcessing/Private/LateReflectionsFast.cpp b/Engine/Source/Runtime/SignalProcessing/Private/LateReflectionsFast.cpp
--- a/Engine/Source/Runtime/SignalProcessing/Private/LateReflectionsFast.cpp
+++ b/Engine/Source/Runtime/SignalProcessing/Private/LateReflectionsFast.cpp
@@ -131,7 +131,8 @@ FLateReflectionsPlate::FLateReflectionsPlate(
 			-Density,
 			PlateDelays.NumSamplesModulatedBase - PlateDelays.NumSamplesModulatedDelta,
 			PlateDelays.NumSamplesModulatedBase + PlateDelays.NumSamplesModulatedDelta,
-			NumInternalBufferSamples);
+			NumInternalBufferSamples,
+			SampleRate);
 	DelayA = MakeUnique<FIntegerDelay>(PlateDelays.NumSamplesDelayA, PlateDelays.NumSamplesDelayA);
 	DelayB = MakeUnique<FIntegerDelay>(PlateDelays.NumSamplesDelayB, PlateDelays.NumSamplesDelayB);
 	DelayC = MakeUnique<FIntegerDelay>(PlateDelays.NumSamplesDelayC, PlateDelays.NumSamplesDelayC);
diff --git a/Engine/Source/Runtime/SignalProcessing/Private/ReverbFast.cpp b/Engine/Source/Runtime/SignalProcessing/Private/ReverbFast.cpp
--- a/Engine/Source/Runtime/SignalProcessing/Private/ReverbFast.cpp
+++ b/Engine/Source/Runtime/SignalProcessing/Private/ReverbFast.cpp
@@ -81,6 +81,12 @@ namespace Audio {
 	// Process a buffer of input audio samples.
 	void FPlateReverbFast::ProcessAudio(const AlignedFloatBuffer& InSamples, const int32 InNumChannels, AlignedFloatBuffer& OutSamples, const int32 OutNumChannels)
 	{
+		if(InSamples.Num() == 0)
+		{
+			OutSamples.Reset(0);
+			return;
+		}
+		
 		ScaledInputBuffer.Reset(InSamples.Num());
 		ScaledInputBuffer.AddUninitialized(InSamples.Num());
 		check(ScaledInputBuffer.Num() == InSamples.Num());
diff --git a/Engine/Source/Runtime/SignalProcessing/Public/DSP/DynamicDelayAPF.h b/Engine/Source/Runtime/SignalProcessing/Public/DSP/DynamicDelayAPF.h
--- a/Engine/Source/Runtime/SignalProcessing/Public/DSP/DynamicDelayAPF.h
+++ b/Engine/Source/Runtime/SignalProcessing/Public/DSP/DynamicDelayAPF.h
@@ -19,7 +19,7 @@ namespace Audio
 		// InMinDelay is the minimum allowable delay of the all pass filter.
 		// InMaxDelay is the maximum allowable delay of the all pass filter.
 		// InMaxNumInternalBufferSamples is the maximum internal block size used internally. 
-		FDynamicDelayAPF(float InG, int32 InMinDelay, int32 InMaxDelay, int32 InMaxNumInternalBufferSamples);
+		FDynamicDelayAPF(float InG, int32 InMinDelay, int32 InMaxDelay, int32 InMaxNumInternalBufferSamples, float InSampleRate);
 
 		// Destructor
 		~FDynamicDelayAPF();
@@ -47,7 +47,7 @@ namespace Audio
 	private:
 		// Feedback/Feedforward gain coefficient
 		FLinearEase G;
-		float EaseTimeInSec{ 2.0 };
+		float EaseTimeInSec{ 2.0f };
 		int32 MinDelay;
 		int32 MaxDelay;
 		int32 NumDelaySamples;
