Dynamic delay apf fixes. Now passes unit tests.

diff --git a/Engine/Source/Runtime/SignalProcessing/Private/DynamicDelayAPF.cpp b/Engine/Source/Runtime/SignalProcessing/Private/DynamicDelayAPF.cpp
--- a/Engine/Source/Runtime/SignalProcessing/Private/DynamicDelayAPF.cpp
+++ b/Engine/Source/Runtime/SignalProcessing/Private/DynamicDelayAPF.cpp
@@ -11,7 +11,7 @@ FDynamicDelayAPF::FDynamicDelayAPF(float InG, int32 InMinDelay, int32 InMaxDelay
 ,	NumDelaySamples(InMinDelay - 1)
 ,	NumInternalBufferSamples(InMaxNumInternalBufferSamples)
 {
-	G.SetValueInterrupt(0.0f);
+	G.SetValueInterrupt(InG);
 	G.Init(InSampleRate);
 
 	checkf(NumDelaySamples >= 0, TEXT("Minimum delay must be atleast 1"));
@@ -101,30 +101,30 @@ void FDynamicDelayAPF::ProcessAudioBlock(const float* InSamples, const AlignedFl
 	// w[n - d]
 	FractionalDelayLine->ProcessAudio(WorkBufferA, InFractionalDelays, WorkBufferB);
 
-	// w[n] = x[n] + G * w[n - d]
 	DelayLineInput.Reset(InNum);
 	DelayLineInput.AddUninitialized(InNum);
 
-
-//	BufferWeightedSumFast(WorkBufferB.GetData(), G, InSamples, DelayLineInput.GetData(), InNum);
+	// Get G values to interpolate over
 	const float LastG = G.GetValue();
 	const float CurrG = G.GetValue(InNum);
 
 	// WorkBufferA = G * WorkBufferB
 	FMemory::Memcpy(WorkBufferA.GetData(), WorkBufferB.GetData(), InNum * sizeof(float));
 	FadeBufferFast(WorkBufferA.GetData(), InNum, LastG, CurrG);
-	SumBuffers(InSamples, WorkBufferA.GetData(), DelayLineInput.GetData(), InNum);
 
-	BufferUnderflowClampFast(WorkBufferB.GetData(), InNum);
+	// DelayLineInput = InSamples + WorkBufferA
+	// w[n] = x[n] + G * w[n - d]
+	SumBuffers(InSamples, WorkBufferA.GetData(), DelayLineInput.GetData(), InNum);
 
-	// y[n] = -G w[n] + w[n - d]
-//	BufferWeightedSumFast(DelayLineInput.GetData(), -G, WorkBufferB.GetData(), OutSamples, InNum);
-	FadeBufferFast(DelayLineInput.GetData(), InNum, -LastG, -CurrG);
-	SumBuffers(DelayLineInput.GetData(), WorkBufferB.GetData(), OutSamples, InNum);
+	BufferUnderflowClampFast(DelayLineInput);
 
 	// Update delay line
 	IntegerDelayLine->RemoveSamples(InNum);
 	IntegerDelayLine->AddSamples(DelayLineInput.GetData(), InNum);
+
+	// y[n] = -G w[n] + w[n - d]
+	FadeBufferFast(DelayLineInput.GetData(), InNum, -LastG, -CurrG);
+	SumBuffers(DelayLineInput.GetData(), WorkBufferB.GetData(), OutSamples, InNum);
 }
 
 void FDynamicDelayAPF::Reset() 
