Adding ability to retrieve current value of linear ease object without updating tick index.

diff --git a/Engine/Source/Runtime/SignalProcessing/Public/DSP/Dsp.h b/Engine/Source/Runtime/SignalProcessing/Public/DSP/Dsp.h
--- a/Engine/Source/Runtime/SignalProcessing/Public/DSP/Dsp.h
+++ b/Engine/Source/Runtime/SignalProcessing/Public/DSP/Dsp.h
@@ -357,6 +357,7 @@ namespace Audio
 			++CurrentTick;
 			return CurrentValue;
 		}
+
 		// same as GetValue(), but overloaded to increment Current Tick by NumTicksToJumpAhead
 		// (before getting the value);
 		float GetValue(int32 NumTicksToJumpAhead)
@@ -372,6 +373,10 @@ namespace Audio
 			return CurrentValue;
 		}
 
+		float GetCurrentValue() const
+		{
+			return CurrentValue;
+		}
 		 
 		// Updates the target value without changing the duration or tick data.
 		// Sets the state as if the new value was the target value all along
