Migrating some DSP from AudioMixer to SignalProcessing
Switching to TArrayViews in some instances
Adding magnitude and power spectrum methods.
Adding some array math routines.

diff --git a/Engine/Source/Runtime/AudioAnalyzer/Private/AudioAnalyzerNRTFacade.cpp b/Engine/Source/Runtime/AudioAnalyzer/Private/AudioAnalyzerNRTFacade.cpp
--- a/Engine/Source/Runtime/AudioAnalyzer/Private/AudioAnalyzerNRTFacade.cpp
+++ b/Engine/Source/Runtime/AudioAnalyzer/Private/AudioAnalyzerNRTFacade.cpp
@@ -83,10 +83,10 @@ namespace Audio
 		}
 
 		// Convert 16 bit pcm to 32 bit float
-		TSampleBuffer<float> FloatSamples((const int16*)InRawWaveData.GetData(), InRawWaveData.Num(), InNumChannels, InSampleRate);
+		TSampleBuffer<float> FloatSamples((const int16*)InRawWaveData.GetData(), InRawWaveData.Num() / 2, InNumChannels, InSampleRate);
 
 		// Perform and finalize audio analysis.
-		Worker->Analyze(FloatSamples, Result.Get());
+		Worker->Analyze(FloatSamples.GetArrayView(), Result.Get());
 		Worker->Finalize(Result.Get());
 
 		return Result;
diff --git a/Engine/Source/Runtime/AudioAnalyzer/Public/IAudioAnalyzerNRTInterface.h b/Engine/Source/Runtime/AudioAnalyzer/Public/IAudioAnalyzerNRTInterface.h
--- a/Engine/Source/Runtime/AudioAnalyzer/Public/IAudioAnalyzerNRTInterface.h
+++ b/Engine/Source/Runtime/AudioAnalyzer/Public/IAudioAnalyzerNRTInterface.h
@@ -8,7 +8,6 @@
 #include "SampleBuffer.h"
 #include "Features/IModularFeature.h"
 
-
 namespace Audio
 {
 	
@@ -75,7 +74,7 @@ namespace Audio
 		 * Perform analysis of an audio stream.
 		 * This method may be called multiple times with audio from the same source.
 		 */
-		virtual void Analyze(const TSampleBuffer<float>& InAudio, IAnalyzerNRTResult* OutResult) = 0;
+		virtual void Analyze(TArrayView<const float> InAudio, IAnalyzerNRTResult* OutResult) = 0;
 
 		/** 
 		 * Called when analysis of audio asset is complete. 
diff --git a/Engine/Source/Runtime/AudioMixer/Private/DSP/FloatArrayMath.cpp b/Engine/Source/Runtime/AudioMixer/Private/DSP/FloatArrayMath.cpp
deleted file mode 100644
diff --git a/Engine/Source/Runtime/AudioMixer/Public/AudioMixerBlueprintLibrary.h b/Engine/Source/Runtime/AudioMixer/Public/AudioMixerBlueprintLibrary.h
--- a/Engine/Source/Runtime/AudioMixer/Public/AudioMixerBlueprintLibrary.h
+++ b/Engine/Source/Runtime/AudioMixer/Public/AudioMixerBlueprintLibrary.h
@@ -58,6 +58,16 @@ enum class EFFTWindowType : uint8
 	Blackman
 };
 
+UENUM(BlueprintType)
+enum class EAudioSpectrumType : uint8
+{
+	// Spectrum frequency values are equal to magnitude of frequency.
+	MagnitudeSpectrum,
+
+	// Spectrum frequency values are equal to magnitude squared.
+	PowerSpectrum
+};
+
 /** 
 * Called when a load request for a sound has completed.
 */
@@ -141,4 +151,4 @@ public:
 
 private:
 	static void PopulateSpectrumAnalyzerSettings(EFFTSize FFTSize, EFFTPeakInterpolationMethod InterpolationMethod, EFFTWindowType WindowType, float HopSize, Audio::FSpectrumAnalyzerSettings &OutSettings);
-};
\ No newline at end of file
+};
diff --git a/Engine/Source/Runtime/AudioMixer/Public/DSP/DeinterleaveView.h b/Engine/Source/Runtime/AudioMixer/Public/DSP/DeinterleaveView.h
deleted file mode 100644
diff --git a/Engine/Source/Runtime/AudioMixer/Public/DSP/FloatArrayMath.h b/Engine/Source/Runtime/AudioMixer/Public/DSP/FloatArrayMath.h
deleted file mode 100644
diff --git a/Engine/Source/Runtime/AudioMixer/Public/DSP/MelScale.h b/Engine/Source/Runtime/AudioMixer/Public/DSP/MelScale.h
--- a/Engine/Source/Runtime/AudioMixer/Public/DSP/MelScale.h
+++ b/Engine/Source/Runtime/AudioMixer/Public/DSP/MelScale.h
@@ -2,7 +2,7 @@
 
 #pragma once
 #include "CoreMinimal.h"
-#include "FloatArrayMath.h"
+#include "DSP/FloatArrayMath.h"
 
 
 namespace Audio
diff --git a/Engine/Source/Runtime/AudioMixer/Public/DSP/SlidingWindow.h b/Engine/Source/Runtime/AudioMixer/Public/DSP/SlidingWindow.h
deleted file mode 100644
diff --git a/Engine/Source/Runtime/SignalProcessing/Private/AlignedBlockBuffer.cpp b/Engine/Source/Runtime/SignalProcessing/Private/AlignedBlockBuffer.cpp
--- a/Engine/Source/Runtime/SignalProcessing/Private/AlignedBlockBuffer.cpp
+++ b/Engine/Source/Runtime/SignalProcessing/Private/AlignedBlockBuffer.cpp
@@ -2,7 +2,6 @@
 
 #include "DSP/AlignedBlockBuffer.h"
 #include "CoreMinimal.h"
-#include "AudioMixer.h"
 #include "SignalProcessingModule.h"
 
 using namespace Audio;
diff --git a/Engine/Source/Runtime/SignalProcessing/Private/AudioFFT.cpp b/Engine/Source/Runtime/SignalProcessing/Private/AudioFFT.cpp
--- a/Engine/Source/Runtime/SignalProcessing/Private/AudioFFT.cpp
+++ b/Engine/Source/Runtime/SignalProcessing/Private/AudioFFT.cpp
@@ -504,6 +504,126 @@ namespace Audio
 		}
 	}
 
+	void ComputePowerSpectrumNoScaling(const FFTFreqDomainData& InFrequencyData, int32 FFTSize, AlignedFloatBuffer& OutBuffer)
+	{
+		check((FFTSize % 2) == 0);
+
+		if (FFTSize < 1)
+		{
+			// Can't do anything with a zero sized fft.
+			OutBuffer.Reset(0);
+			return;
+		}
+
+		// Spectrum only calculates values for real positive frequencies. 
+		const int32 NumSpectrumValues = (FFTSize / 2) + 1;
+
+		// Resize output buffer
+		OutBuffer.Reset(NumSpectrumValues);
+		OutBuffer.AddUninitialized(NumSpectrumValues);
+
+		float* OutBufferData = OutBuffer.GetData();
+
+		const float* RealData = InFrequencyData.OutReal;
+		const float* ImagData = InFrequencyData.OutImag;
+
+		if (IsAligned<const float*>(RealData, AUDIO_BUFFER_ALIGNMENT) && IsAligned<const float*>(ImagData, AUDIO_BUFFER_ALIGNMENT) && (NumSpectrumValues > 5))
+		{
+			BufferComplexToPowerFast(RealData, ImagData, OutBufferData, NumSpectrumValues - 1);
+
+			// Buffer operations are on 16 byte boundaries, which excludes the location of the nyquist
+			// frequency in the fft output buffers. Here we explicitly handle data at nyquist freq.
+			const int32 NyquistIndex = NumSpectrumValues - 1;
+			const float NyquistReal = RealData[NyquistIndex];
+			const float NyquistImag = ImagData[NyquistIndex];
+			OutBufferData[NyquistIndex] = NyquistReal * NyquistReal + NyquistImag * NyquistImag;
+		}
+		else
+		{
+			for (int32 i = 0; i < NumSpectrumValues; i++)
+			{
+				OutBufferData[i] = RealData[i] * RealData[i] + ImagData[i] * ImagData[i];
+			}
+		}
+	}
+
+	void ComputePowerSpectrum(const FFTFreqDomainData& InFrequencyData, int32 FFTSize, AlignedFloatBuffer& OutBuffer)
+	{
+		if (FFTSize < 1)
+		{
+			OutBuffer.Reset();
+			return;
+		}
+
+		ComputePowerSpectrumNoScaling(InFrequencyData, FFTSize, OutBuffer);
+
+		const int32 NumSpectrumValues = OutBuffer.Num();
+		if (NumSpectrumValues < 1)
+		{
+			return;
+		}
+
+		const float FFTScale = 1.f / static_cast<float>(FFTSize);
+
+		float* OutBufferData = OutBuffer.GetData();
+
+		if (NumSpectrumValues > 1)
+		{
+			MultiplyBufferByConstantInPlace(OutBufferData, NumSpectrumValues - 1, FFTScale);
+		}
+		
+		OutBufferData[NumSpectrumValues - 1] *= FFTScale;
+	}
+
+	void ComputeMagnitudeSpectrum(const FFTFreqDomainData& InFrequencyData, int32 FFTSize, AlignedFloatBuffer& OutBuffer)
+	{
+		if (FFTSize < 1)
+		{
+			OutBuffer.Reset();
+			return;
+		}
+
+		ComputePowerSpectrumNoScaling(InFrequencyData, FFTSize, OutBuffer);
+
+		const int32 NumSpectrumValues = OutBuffer.Num();
+		if (NumSpectrumValues < 1)
+		{
+			return;
+		}
+
+		const float FFTScale = 1.f / FMath::Sqrt(static_cast<float>(FFTSize));
+
+		float* OutBufferData = OutBuffer.GetData();
+
+		for (int32 i = 0; i < NumSpectrumValues; i++)
+		{
+			// TODO: Currently no vector sqrt. utilize fmath for now.
+			OutBufferData[i] = FMath::Sqrt(OutBufferData[i]);
+		}
+
+		if (NumSpectrumValues > 1)
+		{
+			MultiplyBufferByConstantInPlace(OutBufferData, NumSpectrumValues - 1, FFTScale);
+		}
+		OutBufferData[NumSpectrumValues - 1] *= FFTScale;
+	}
+
+	void ComputeSpectrum(ESpectrumType InSpectrumType, const FFTFreqDomainData& InFrequencyData, int32 FFTSize, AlignedFloatBuffer& OutBuffer)
+	{
+		switch (InSpectrumType)
+		{
+			case ESpectrumType::MagnitudeSpectrum:
+				ComputeMagnitudeSpectrum(InFrequencyData, FFTSize, OutBuffer);
+				break;
+
+			case ESpectrumType::PowerSpectrum:
+				ComputePowerSpectrum(InFrequencyData, FFTSize, OutBuffer);
+				break;
+
+			default:
+				checkf(false, TEXT("Unhandled Audio::ESpectrumType"));
+		}
+	}
 	SIGNALPROCESSING_API void CrossCorrelate(AlignedFloatBuffer& FirstBuffer, AlignedFloatBuffer& SecondBuffer, AlignedFloatBuffer& OutCorrelation, bool bZeroPad /*= true*/)
 	{
 		FrequencyBuffer OutputCorrelationFrequencies;
diff --git a/Engine/Source/Runtime/SignalProcessing/Private/BufferVectorOperations.cpp b/Engine/Source/Runtime/SignalProcessing/Private/BufferVectorOperations.cpp
--- a/Engine/Source/Runtime/SignalProcessing/Private/BufferVectorOperations.cpp
+++ b/Engine/Source/Runtime/SignalProcessing/Private/BufferVectorOperations.cpp
@@ -1850,4 +1850,40 @@ namespace Audio
 			OutPos += 4;
 		}
 	}
+
+	void BufferComplexToPowerFast(const AlignedFloatBuffer& InRealSamples, const AlignedFloatBuffer& InImaginarySamples, AlignedFloatBuffer& OutPowerSamples)
+	{
+		checkf(InRealSamples.Num() == InImaginarySamples.Num(), TEXT("Input buffers must have equal number of elements"));
+
+		const int32 Num = InRealSamples.Num();
+
+		OutPowerSamples.Reset(Num);
+		OutPowerSamples.AddUninitialized(Num);
+
+		BufferComplexToPowerFast(InRealSamples.GetData(), InImaginarySamples.GetData(), OutPowerSamples.GetData(), Num);
+	}
+
+	void BufferComplexToPowerFast(const float* RESTRICT InRealSamples, const float* RESTRICT InImaginarySamples, float* RESTRICT OutPowerSamples, const int32 InNum)
+	{
+		checkf(InNum >= 4, TEXT("Buffer must have atleast 4 elements."));
+		checkf(0 == (InNum % 4), TEXT("Buffer length be a multiple of 4."));
+
+		checkf(IsAligned<const float*>(InRealSamples, AUDIO_SIMD_FLOAT_ALIGNMENT), TEXT("Memory must be aligned to use vector operations."));
+		checkf(IsAligned<const float*>(InImaginarySamples, AUDIO_SIMD_FLOAT_ALIGNMENT), TEXT("Memory must be aligned to use vector operations."));
+		checkf(IsAligned<float*>(OutPowerSamples, AUDIO_SIMD_FLOAT_ALIGNMENT), TEXT("Memory must be aligned to use vector operations."));
+
+		
+		for (int32 i = 0; i < InNum; i += 4)
+		{
+			VectorRegister VInReal = VectorLoadAligned(&InRealSamples[i]);
+			VectorRegister VInRealSquared = VectorMultiply(VInReal, VInReal);
+
+			VectorRegister VInImag = VectorLoadAligned(&InImaginarySamples[i]);
+			VectorRegister VInImagSquared = VectorMultiply(VInImag, VInImag);
+
+			VectorRegister VOut = VectorAdd(VInRealSquared, VInImagSquared);
+
+			VectorStoreAligned(VOut, &OutPowerSamples[i]);
+		}
+	}
 }
diff --git a/Engine/Source/Runtime/SignalProcessing/Private/FloatArrayMath.cpp b/Engine/Source/Runtime/SignalProcessing/Private/FloatArrayMath.cpp
new file mode 100644
--- /dev/null
+++ b/Engine/Source/Runtime/SignalProcessing/Private/FloatArrayMath.cpp
@@ -0,0 +1,327 @@
+// Copyright 1998-2019 Epic Games, Inc. All Rights Reserved.
+
+#include "DSP/FloatArrayMath.h"
+#include "CoreMinimal.h"
+
+namespace
+{
+	const float LOGE10 = FMath::Loge(10.f);
+}
+
+namespace Audio
+{
+	void ArrayCumulativeSum(const TArray<float>& InData, TArray<float>& OutData)
+	{
+		// Initialize output data
+		int32 Num = InData.Num();
+		OutData.Reset();
+		OutData.AddUninitialized(Num);
+
+		if (Num < 1)
+		{
+			return;
+		}
+
+		float* OutDataPtr = OutData.GetData();
+		const float* InDataPtr = InData.GetData();
+
+		// Start summing
+		*OutDataPtr = *InDataPtr++;
+
+		for (int32 i = 1; i < Num; i++)
+		{
+			float Temp = *OutDataPtr++ + *InDataPtr++;
+			*OutDataPtr = Temp;
+		}
+	}
+
+	void ArrayMeanFilter(const TArray<float>& InData, int32 WindowSize, int32 WindowOrigin, TArray<float>& OutData)
+	{
+		// a quick but sinful implementation of a mean filter. encourages floating point rounding errors. 
+		check(WindowOrigin < WindowSize);
+		check(WindowOrigin >= 0);
+		check(WindowSize > 0);
+
+		// Initialize output data
+		const int32 Num = InData.Num();
+		OutData.Reset();
+		OutData.AddUninitialized(Num);
+
+		if (Num < 1)
+		{
+			return;
+		}
+		
+		// Use cumulative sum to avoid multiple summations 
+		// Instead of summing over InData[StartIndex:EndIndex], avoid all that
+		// calculation by taking difference of cumulative sum at those two points:
+		//  cumsum(X[0:b]) - cumsum(X[0:a]) = sum(X[a:b])
+		TArray<float> SummedData;
+		ArrayCumulativeSum(InData, SummedData);
+		const float LastSummedData = SummedData.Last();
+		
+		
+		const int32 LastIndexBeforeEndBoundaryCondition = FMath::Max(WindowOrigin + 1, Num - WindowSize + WindowOrigin + 1);
+		const int32 StartOffset = -WindowOrigin - 1;
+		const int32 EndOffset = WindowSize - WindowOrigin - 1;
+		const int32 WindowTail = WindowSize - WindowOrigin;
+
+		float* OutDataPtr = OutData.GetData();
+		const float* SummedDataPtr = SummedData.GetData();
+
+		if ((WindowSize - WindowOrigin) < Num)
+		{
+			// Handle boundary condition where analysis window precedes beginning of array.
+			for (int32 i = 0; i < (WindowOrigin + 1); i++)
+			{
+				OutDataPtr[i] = SummedDataPtr[i + EndOffset] / FMath::Max(1.f, static_cast<float>(WindowTail + i));
+			}
+
+			// No boundary conditions to handle here.	
+			const float MeanDivisor = static_cast<float>(WindowSize);
+			for (int32 i = WindowOrigin + 1; i < LastIndexBeforeEndBoundaryCondition; i++)
+			{
+				OutDataPtr[i] = (SummedDataPtr[i + EndOffset] - SummedDataPtr[i + StartOffset]) / MeanDivisor;
+			}
+		}
+		else
+		{
+			// Handle boundary condition where window precedes beginning and goes past end of array
+			const float ArrayMean = LastSummedData / static_cast<float>(Num);
+			for (int32 i = 0; i < LastIndexBeforeEndBoundaryCondition; i++)
+			{
+				OutDataPtr[i] = ArrayMean;
+			}
+		}
+
+		// Handle boundary condition where analysis window goes past end of array.
+		for (int32 i = LastIndexBeforeEndBoundaryCondition; i < Num; i++)
+		{
+			OutDataPtr[i] = (LastSummedData - SummedDataPtr[i + StartOffset]) / static_cast<float>(Num - i + WindowOrigin);
+		}
+	}
+
+	void ArrayMaxFilter(const TArray<float>& InData, int32 WindowSize, int32 WindowOrigin, TArray<float>& OutData)
+	{
+		// A reasonable implementation of a max filter for the data we're interested in, though surely not the fastest.
+		check(WindowOrigin < WindowSize);
+		check(WindowOrigin >= 0);
+		check(WindowSize > 0);
+		
+		int32 StartIndex = -WindowOrigin;
+		int32 EndIndex = StartIndex + WindowSize;
+
+		// Initialize output
+		int32 Num = InData.Num();
+		OutData.Reset();
+		OutData.AddUninitialized(Num);
+
+		if (Num < 1)
+		{
+			return;
+		}
+
+		// Get max in first window
+		int32 ActualStartIndex = 0;
+		int32 ActualEndIndex = FMath::Min(EndIndex, Num);
+
+		const float* InDataPtr = InData.GetData();
+		float* OutDataPtr = OutData.GetData();
+		int32 MaxIndex = 0;
+		float MaxValue = InData[0];
+
+		for (int32 i = ActualStartIndex; i < ActualEndIndex; i++)
+		{
+			if (InDataPtr[i] > MaxValue)
+			{
+				MaxValue = InDataPtr[i];
+				MaxIndex = i;
+			}		
+		}
+		OutDataPtr[0] = MaxValue;
+
+		StartIndex++;
+		EndIndex++;
+
+		// Get max in remaining windows
+		for (int32 i = 1; i < Num; i++)
+		{
+			ActualStartIndex = FMath::Max(StartIndex, 0);
+			ActualEndIndex = FMath::Min(EndIndex, Num);
+
+			if (MaxIndex < StartIndex)
+			{
+				// We need to evaluate the entire window because the previous maximum value was not in this window.
+				MaxIndex = ActualStartIndex;
+				MaxValue = InDataPtr[MaxIndex];
+				for (int32 j = ActualStartIndex + 1; j < ActualEndIndex; j++)
+				{
+					if (InDataPtr[j] > MaxValue)
+					{
+						MaxIndex = j;
+						MaxValue = InDataPtr[MaxIndex];
+					}
+				}
+			}
+			else
+			{
+				// We only need to inspect the newest sample because the previous maximum value was in this window.
+				if (InDataPtr[ActualEndIndex - 1] > MaxValue)
+				{
+					MaxIndex = ActualEndIndex - 1;
+					MaxValue = InDataPtr[MaxIndex];
+				}
+			}
+
+			OutDataPtr[i] = MaxValue;
+
+			StartIndex++;
+			EndIndex++;
+		}
+	}
+
+	void ArrayGetEuclideanNorm(const TArray<float>& InArray, float& OutEuclideanNorm)
+	{
+		// Initialize output.
+		OutEuclideanNorm = 0.0f;
+		const int32 Num = InArray.Num();
+		const float* InArrayData = InArray.GetData();
+		
+		// Sum it up.
+		for (int32 i = 0; i < Num; i++)
+		{
+			OutEuclideanNorm += InArrayData[i] * InArrayData[i];
+		}
+
+		OutEuclideanNorm = FMath::Sqrt(OutEuclideanNorm);
+	}
+
+	void ArrayMultiplyByConstantInPlace(TArray<float>& InArray, float InMultiplier)
+	{
+		const int32 Num = InArray.Num();
+		float* InArrayData = InArray.GetData();
+		for (int32 i = 0; i < Num; i++)
+		{
+			InArrayData[i] *= InMultiplier;
+		}
+	}
+
+	void ArraySubtractByConstantInPlace(TArray<float>& InArray, float InSubtrahend)
+	{
+		const int32 Num = InArray.Num();
+		float* InArrayData = InArray.GetData();
+		for (int32 i = 0; i < Num; i++)
+		{
+			InArrayData[i] -= InSubtrahend;
+		}
+	}
+
+	void ArrayMagnitudeToDecibelInPlace(TArray<float>& InArray)
+	{
+		const int32 Num = InArray.Num();
+		float* InArrayData = InArray.GetData();
+		for (int32 i = 0; i < Num; i++)
+		{
+			InArrayData[i] = 20.f * FMath::Loge(InArrayData[i]) / LOGE10;
+		}
+	}
+
+	void ArrayPowerToDecibelInPlace(TArray<float>& InArray)
+	{
+		const int32 Num = InArray.Num();
+		float* InArrayData = InArray.GetData();
+		for (int32 i = 0; i < Num; i++)
+		{
+			InArrayData[i] = 10.f * FMath::Loge(InArrayData[i]) / LOGE10;
+		}
+	}
+
+	FContiguousSparse2DKernelTransform::FContiguousSparse2DKernelTransform(const int32 NumInElements, const int32 NumOutElements)
+	:	NumIn(NumInElements)
+	,	NumOut(NumOutElements)
+	{
+		check(NumIn >= 0);
+		check(NumOut >= 0)
+		FRow EmptyRow;
+		EmptyRow.StartIndex = 0;
+
+		// Fill up the kernel with emptp rows
+		Kernel.Init(EmptyRow, NumOut);
+	}
+
+	int32 FContiguousSparse2DKernelTransform::GetNumInElements() const
+	{
+		return NumIn;
+	}
+
+
+	int32 FContiguousSparse2DKernelTransform::GetNumOutElements() const
+	{
+		return NumOut;
+	}
+
+	void FContiguousSparse2DKernelTransform::SetRow(const int32 RowIndex, const int32 StartIndex, TArrayView<const float> OffsetValues)
+	{
+		check((StartIndex + OffsetValues.Num()) <= NumIn);
+
+		// Copy row data internally
+		Kernel[RowIndex].StartIndex = StartIndex;
+		Kernel[RowIndex].OffsetValues = TArray<float>(OffsetValues.GetData(), OffsetValues.Num());
+	}
+
+	void FContiguousSparse2DKernelTransform::TransformArray(TArrayView<const float> InArray, TArray<float>& OutArray) const
+	{
+		check(InArray.Num() == NumIn);
+
+		// Resize output
+		OutArray.Reset(NumOut);
+		if (NumOut > 0)
+		{
+			OutArray.AddUninitialized(NumOut);
+		}
+
+		TransformArray(InArray.GetData(), OutArray.GetData());
+	}
+
+	void FContiguousSparse2DKernelTransform::TransformArray(TArrayView<const float> InArray, AlignedFloatBuffer& OutArray) const
+	{	
+		check(InArray.Num() == NumIn);
+
+		// Resize output
+		OutArray.Reset(NumOut);
+		if (NumOut > 0)
+		{
+			OutArray.AddUninitialized(NumOut);
+		}
+
+		TransformArray(InArray.GetData(), OutArray.GetData());
+	}
+
+	void FContiguousSparse2DKernelTransform::TransformArray(const float* InArray, float* OutArray) const
+	{
+		check(nullptr != InArray);
+		check(nullptr != OutArray);
+
+		// Initialize output
+		FMemory::Memset(OutArray, 0, sizeof(float) * NumOut);
+
+		// Apply kernel one row at a time
+		const FRow* KernelData = Kernel.GetData();
+		for (int32 RowIndex = 0; RowIndex < Kernel.Num(); RowIndex++)
+		{
+			const FRow& Row = KernelData[RowIndex];
+
+			// Get offset pointer into input array.
+			const float* OffsetInData = &InArray[Row.StartIndex];
+			// Get offset pointer of row.
+			const float* RowValuePtr = Row.OffsetValues.GetData();
+
+			// dot prod 'em. 
+			int32 NumToMult = Row.OffsetValues.Num();
+			for (int32 i = 0; i < NumToMult; i++)
+			{
+				OutArray[RowIndex] += OffsetInData[i] * RowValuePtr[i];
+			}
+		}
+	}
+}
diff --git a/Engine/Source/Runtime/SignalProcessing/Public/DSP/AudioFFT.h b/Engine/Source/Runtime/SignalProcessing/Public/DSP/AudioFFT.h
--- a/Engine/Source/Runtime/SignalProcessing/Public/DSP/AudioFFT.h
+++ b/Engine/Source/Runtime/SignalProcessing/Public/DSP/AudioFFT.h
@@ -8,6 +8,13 @@
 
 namespace Audio
 {
+	// Types of spectrums which can be directly derived from FFTFreqDomainData
+	enum class ESpectrumType : uint8
+	{
+		MagnitudeSpectrum,
+		PowerSpectrum
+	};
+
 	namespace FFTIntrinsics
 	{
 		SIGNALPROCESSING_API uint32 NextPowerOf2(uint32 Input);
@@ -207,8 +214,36 @@ namespace Audio
 		FrequencyBuffer InputFrequencies;
 		int32 BlockSize;
 
+
 		AlignedFloatBuffer TimeDomainInputBuffer;
 		AlignedFloatBuffer COLABuffer;
 	};
 
+	// Computes the power spectrum from FFTFreqDomainData. Applies a 1/(FFTSize^2) scaling to the output to 
+	// maintain equal energy between original time domain data and output spectrum.  Only the first 
+	// (FFTSize / 2 + 1) spectrum values are calculated. These represent the frequencies from 0 to Nyquist.
+	//
+	// InFrequencyData is the input frequency domain data. Generally this is created by calling PerformFFT(...)
+	// FFTSize is the number of samples used when originally calculating the FFT
+	// OutBuffer is an aligned buffer which will contain spectrum data. It will constain (FFTSize / 2 + 1) elements.
+	SIGNALPROCESSING_API void ComputePowerSpectrum(const FFTFreqDomainData& InFrequencyData, int32 FFTSize, AlignedFloatBuffer& OutBuffer);
+
+	// Computes the magnitude spectrum from FFTFreqDomainData. Applies a 1/FFTSize scaling to the output to 
+	// maintain equal energy between original time domain data and output spectrum.  Only the first 
+	// (FFTSize / 2 + 1) spectrum values are calculated. These represent the frequencies from 0 to Nyquist.
+	//
+	// InFrequencyData is the input frequency domain data. Generally this is created by calling PerformFFT(...)
+	// FFTSize is the number of samples used when originally calculating the FFT
+	// OutBuffer is an aligned buffer which will contain spectrum data. It will constain (FFTSize / 2 + 1) elements.
+	SIGNALPROCESSING_API void ComputeMagnitudeSpectrum(const FFTFreqDomainData& InFrequencyData, int32 FFTSize, AlignedFloatBuffer& OutBuffer);
+	
+	// Computes the spectrum from FFTFreqDomainData. Applies a scaling to the output to maintain equal 
+	// energy between original time domain data and output spectrum.  Only the first (FFTSize / 2 + 1)
+	// spectrum values are calculated. These represent the frequencies from 0 to Nyquist.
+	//
+	// InSpectrumType denotes which spectrum type to calculate.
+	// InFrequencyData is the input frequency domain data. Generally this is created by calling PerformFFT(...)
+	// FFTSize is the number of samples used when originally calculating the FFT
+	// OutBuffer is an aligned buffer which will contain spectrum data. It will constain (FFTSize / 2 + 1) elements.
+	SIGNALPROCESSING_API void ComputeSpectrum(ESpectrumType InSpectrumType, const FFTFreqDomainData& InFrequencyData, int32 FFTSize, AlignedFloatBuffer& OutBuffer);
 }
diff --git a/Engine/Source/Runtime/SignalProcessing/Public/DSP/BufferVectorOperations.h b/Engine/Source/Runtime/SignalProcessing/Public/DSP/BufferVectorOperations.h
--- a/Engine/Source/Runtime/SignalProcessing/Public/DSP/BufferVectorOperations.h
+++ b/Engine/Source/Runtime/SignalProcessing/Public/DSP/BufferVectorOperations.h
@@ -16,10 +16,12 @@
 
 namespace Audio
 {
-	
-	typedef TArray<float, TAlignedHeapAllocator<AUDIO_BUFFER_ALIGNMENT>> AlignedFloatBuffer;
-	typedef TArray<uint8, TAlignedHeapAllocator<AUDIO_BUFFER_ALIGNMENT>> AlignedByteBuffer;
-	typedef TArray<int32, TAlignedHeapAllocator<AUDIO_BUFFER_ALIGNMENT>> AlignedInt32Buffer;
+	/** Aligned allocator used for fast operations. */	
+	typedef TAlignedHeapAllocator<AUDIO_BUFFER_ALIGNMENT> FAudioBufferAlignedAllocator;
+
+	typedef TArray<float, FAudioBufferAlignedAllocator> AlignedFloatBuffer;
+	typedef TArray<uint8, FAudioBufferAlignedAllocator> AlignedByteBuffer;
+	typedef TArray<int32, FAudioBufferAlignedAllocator> AlignedInt32Buffer;
 
 
 	/** CHANNEL-AGNOSTIC OPERATIONS */
@@ -245,5 +247,17 @@ namespace Audio
 	SIGNALPROCESSING_API void BufferSum2ChannelToMonoFast(const AlignedFloatBuffer& InSamples, AlignedFloatBuffer& OutSamples);
 
 	/** Sums 2 channel interleaved input samples. OutSamples[n] = InSamples[2n] + InSamples[2n + 1] */
-	SIGNALPROCESSING_API void BufferSum2ChannelToMonoFast(const float* RESTRICT InSamples, float* RESTRICT OutSamples, const int32 InNumFrames);
+	void SIGNALPROCESSING_API BufferSum2ChannelToMonoFast(const float* RESTRICT InSamples, float* RESTRICT OutSamples, const int32 InNumFrames);
+
+	/** Compute power of complex data. Out[i] = Real[i] * Real[i] + Imaginary[i] * Imaginary[i] */
+	void SIGNALPROCESSING_API BufferComplexToPowerFast(const AlignedFloatBuffer& InRealSamples, const AlignedFloatBuffer& InImaginarySamples, AlignedFloatBuffer& OutPowerSamples);
+
+	/** Compute power of complex data. Out[i] = Real[i] * Real[i] + Imaginary[i] * Imaginary[i] */
+	void SIGNALPROCESSING_API BufferComplexToPowerFast(const float* RESTRICT InRealSamples, const float* RESTRICT InImaginarySamples, float* RESTRICT OutPowerSamples, const int32 InNum);
+
+	/** Compute magnitude of complex data. Out[i] = Sqrt(Real[i] * Real[i] + Imaginary[i] * Imaginary[i]) */
+	void SIGNALPROCESSING_API BufferComplexToMagnitudeFast(const AlignedFloatBuffer& InRealSamples, const AlignedFloatBuffer& InImaginarySamples, AlignedFloatBuffer& OutPowerSamples);
+
+	/** Compute magnitude of complex data. Out[i] = Sqrt(Real[i] * Real[i] + Imaginary[i] * Imaginary[i]) */
+	void SIGNALPROCESSING_API BufferComplexToMagnitudeFast(const float* RESTRICT InRealSamples, const float* RESTRICT InImaginarySamples, float* RESTRICT OutPowerSamples, const int32 InNum);
 }
diff --git a/Engine/Source/Runtime/SignalProcessing/Public/DSP/DeinterleaveView.h b/Engine/Source/Runtime/SignalProcessing/Public/DSP/DeinterleaveView.h
new file mode 100644
--- /dev/null
+++ b/Engine/Source/Runtime/SignalProcessing/Public/DSP/DeinterleaveView.h
@@ -0,0 +1,213 @@
+// Copyright 1998-2019 Epic Games, Inc. All Rights Reserved.
+
+#pragma once
+
+#include "CoreMinimal.h"
+
+namespace Audio
+{
+	/** TDeinterleaveView
+	 *
+	 *  TDeinterleaveView provides iterators to access deinterleaved arrays from an interleaved array.
+	 */
+	template <typename T>
+	class TDeinterleaveView
+	{
+		friend class TChannelIterator;
+
+	public:
+
+		/** TChannel
+		 *
+		 *  TChannel provides a contiguous copy of an single channel from an interleaved array.
+		 */
+		template<typename InAllocator = FDefaultAllocator>
+		class TChannel
+		{
+			public:
+				const TArray<T, InAllocator>& Values;
+				const int32 ChannelIndex;
+
+				TChannel(const TArray<T, InAllocator>& InValues, const int32 InChannelIndex)
+				:	Values(InValues)
+				,	ChannelIndex(InChannelIndex)
+				{}
+		};
+
+		/** TChannelIterator
+		 *
+		 *  TChannelIterator iterates over channels in an interleaved array and providing
+		 *  contiguous arrays of a single channel.
+		 */
+		template<typename InAllocator = FDefaultAllocator>
+		class TChannelIterator
+		{
+			TDeinterleaveView<T> DeinterleaveView;
+			TArray<T, InAllocator>& ArrayToFill;
+			int32 ChannelIndex;
+
+		public:
+
+			/** Denotes the end of a channel iterator */
+			static const int32 ChannelIndexEnd = INDEX_NONE;
+
+			/** TChannelIterator
+			 *
+			 *  TChannelIterator allows iteration over interleaved arrays and provides contiguous arrays of single channels.
+			 *
+			 *  InDeinterleaveView is the parent view provides the source interleaved array for this iterator.
+			 *  InArrayToFill is the array which will be populated with contiguous elements for a channel.
+			 *  InChannelIndex is the channel index which this iterator will point to.
+			 */
+			TChannelIterator(TDeinterleaveView<T> InDeinterleaveView, TArray<T, InAllocator>& InArrayToFill, int32 InChannelIndex)
+			:	DeinterleaveView(InDeinterleaveView)
+			,	ArrayToFill(InArrayToFill)
+			,	ChannelIndex(InChannelIndex)
+			{
+				if (ChannelIndex >= DeinterleaveView.NumChannels)
+				{
+					// Channel index is more than number of total channels
+					ChannelIndex = ChannelIndexEnd;
+				}
+
+
+				// prepare the array
+				ArrayToFill.Reset(DeinterleaveView.NumElementsPerChannel);
+				if (DeinterleaveView.NumElementsPerChannel > 0)
+				{
+					ArrayToFill.AddUninitialized(DeinterleaveView.NumElementsPerChannel);
+				}
+			}
+
+			/** Increment the iterator forward by one channel */
+			TChannelIterator& operator++()
+			{
+				if (ChannelIndex != ChannelIndexEnd)
+				{
+					ChannelIndex++;
+					if (ChannelIndex >= DeinterleaveView.NumChannels)
+					{
+						ChannelIndex = ChannelIndexEnd;
+					}
+				}
+				return *this;
+			}
+
+			/** Check equality between iterators */
+			bool operator!=(const TChannelIterator& Other) const
+			{
+				return Other.ChannelIndex != ChannelIndex;
+			}
+
+			/** Get the current channel index */
+			int32 GetChannelIndex() const
+			{
+				return ChannelIndex;
+			}
+
+			/** Dereference the iterator, returns a TChannel object */
+			TChannel<InAllocator> operator*()
+			{
+				if (ChannelIndex == ChannelIndexEnd)
+				{
+					ArrayToFill.Reset(DeinterleaveView.NumElementsPerChannel);
+				}
+				else
+				{
+					// Fill array with deinterleave data
+					T* ArrayToFillData = ArrayToFill.GetData();
+					const T* InterleavedData = DeinterleaveView.InterleavedArray.GetData();
+
+					int32 InterleavedPos = ChannelIndex;
+					for (int32 OutPos = 0; OutPos < DeinterleaveView.NumElementsPerChannel; OutPos++, InterleavedPos += DeinterleaveView.NumChannels)
+					{
+						ArrayToFillData[OutPos] = InterleavedData[InterleavedPos];
+					}
+				}
+
+
+				// Create and return channel
+				return TChannel<InAllocator>(ArrayToFill, ChannelIndex);
+			}
+		};
+
+		/** TDeinterleave constructor.
+		 *
+		 *  InInterleavedArray is the interleaved array to be deinterleaved.
+		 *  InNumChannels is the number of channels in the deinterleaved array.
+		 */
+		TDeinterleaveView(TArrayView<const T> InInterleavedArray, int32 InNumChannels)
+		:	InterleavedArray(InInterleavedArray)
+		,	NumChannels(InNumChannels)
+		,	NumElementsPerChannel(0)
+		{
+			check(InterleavedArray.Num() % NumChannels == 0);
+
+			NumElementsPerChannel = InterleavedArray.Num() / NumChannels;
+		}
+
+		/** Return an STL iterator to the first channel. It fills the given array with channel elements */
+		template <typename InAllocator = FDefaultAllocator>
+		TChannelIterator<InAllocator> begin(TArray<T, InAllocator>& InArrayToFill) const
+		{
+			return TChannelIterator<InAllocator>(*this, InArrayToFill, 0);
+		}
+
+		/** Return an STL iterator to the end. */
+		template <typename InAllocator = FDefaultAllocator>
+		TChannelIterator<InAllocator> end(TArray<T, InAllocator>& InArrayToFill) const
+		{
+			return TChannelIterator<InAllocator>(*this, InArrayToFill, TChannelIterator<InAllocator>::ChannelIndexEnd);
+		}
+
+	private:
+		// view to interleaved data
+		TArrayView<const T> InterleavedArray;
+
+		int32 NumChannels;
+		int32 NumElementsPerChannel;
+	};
+
+	/** TAutoDeinterleaveView
+	 *
+	 *  TAutoDeinterlaveView provides a STL like iterators which exposes contiguous channel arrays from interleaved arrays.  As opposed to TDeinterleaveView, this class can be used in range based for loops, but only one iterator is valid at a time since they all share the same InArrayToFill.
+	 *
+	 *  Example:
+	 *
+	 * 	TArray<float> ArrayToFill;
+	 *  for (auto Channel : TAudoDeineterleaveView(InterleavedArray, ArrayToFill, 2))
+	 *  {
+	 *  	DoSomethingWithAudio(Channel.Values, Channel.ChannelIndex);
+	 *  }
+	 */
+	template <typename T, typename InAllocator = FDefaultAllocator> 
+	class TAutoDeinterleaveView : public TDeinterleaveView<T>
+	{
+		TArray<T, InAllocator>& ArrayToFill;
+
+	public:
+		/** TAutoDeinterleaveView Constructor.
+		 *
+		 *  InInterleavedArray is the interleaved array to be deinterleaved.
+		 *  InArrayToFill is the array which will be populated with contiguous elements for a channel.
+		 *  InNumChannels is the number of channels in the deinterleaved array.
+		 */
+		TAutoDeinterleaveView(TArrayView<const T> InInterleavedArray, TArray<T, InAllocator>& InArrayToFill, int32 InNumChannels)
+		:	TDeinterleaveView(InInterleavedArray, InNumChannels)
+		,	ArrayToFill(InArrayToFill)
+		{}
+
+		/** Return an STL iterator to the first channel. */
+		TChannelIterator<InAllocator> begin() 
+		{
+			return TDeinterleaveView<T>::begin<InAllocator>(ArrayToFill);
+		}
+
+		/** Return an STL iterator to the end. */
+		TChannelIterator<InAllocator> end()
+		{
+			return TDeinterleaveView<T>::end<InAllocator>(ArrayToFill);
+		}
+	};
+}
+
diff --git a/Engine/Source/Runtime/SignalProcessing/Public/DSP/FloatArrayMath.h b/Engine/Source/Runtime/SignalProcessing/Public/DSP/FloatArrayMath.h
new file mode 100644
--- /dev/null
+++ b/Engine/Source/Runtime/SignalProcessing/Public/DSP/FloatArrayMath.h
@@ -0,0 +1,124 @@
+// Copyright 1998-2019 Epic Games, Inc. All Rights Reserved.
+
+#pragma once
+
+#include "CoreMinimal.h"
+#include "DSP/BufferVectorOperations.h"
+
+namespace Audio
+{
+	/** Cumulative sum of array.
+	 *
+	 *  InData contains data to be cumulatively summed.
+	 *  OutData contains sum and is same size as InData.
+	 */
+	SIGNALPROCESSING_API void ArrayCumulativeSum(const TArray<float>& InData, TArray<float>& OutData);
+
+	/** Mean filter of array.
+	 *
+	 *  Note: Uses standard biased mean estimator of Sum(x) / Count(x).
+	 *  Note: At array boundaries, this algorithm truncates windows where no valid array data exists. Values calculated with truncated windows have corresponding increased variances. 
+	 *
+	 *  InData contains data to be filtered.
+	 *  WindowSize determines the number of samples from InData analyzed to produce a value in OutData.
+	 *  WindowOrigin describes the offset from the windows first sample to the index of OutData. For example, if WindowOrigin = WindowSize/4, then OutData[i] = Mean(InData[i - Window/4 : i + 3 * Window / 4]).
+	 *  OutData contains the produceds data.
+	 */
+	SIGNALPROCESSING_API void ArrayMeanFilter(const TArray<float>& InData, int32 WindowSize, int32 WindowOrigin, TArray<float>& OutData);
+
+	/** Max filter of array.
+	 *
+	 *  Note: At array boundaries, this algorithm truncates windows where no valid array data exists. 
+	 *
+	 *  InData contains data to be filtered.
+	 *  WindowSize determines the number of samples from InData analyzed to produce a value in OutData.
+	 *  WindowOrigin describes the offset from the windows first sample to the index of OutData. For example, if WindowOrigin = WindowSize/4, then OutData[i] = Max(InData[i - Window/4 : i + 3 * Window / 4]).
+	 *  OutData contains the produceds data.
+	 */
+	SIGNALPROCESSING_API void ArrayMaxFilter(const TArray<float>& InData, int32 WindowSize, int32 WindowOrigin, TArray<float>& OutData);
+
+	/** Computes the EuclideanNorm of the InArray. Same as calculating the energy in window. */
+	SIGNALPROCESSING_API void ArrayGetEuclideanNorm(const TArray<float>& InArray, float& OutEuclideanNorm);
+
+	/** Multiplies each element in InArray by InMultiplier */
+	SIGNALPROCESSING_API void ArrayMultiplyByConstantInPlace(TArray<float>& InArray, float InMultiplier);
+
+	/** Subtract value from each element in InArray */
+	SIGNALPROCESSING_API void ArraySubtractByConstantInPlace(TArray<float>& InArray, float InSubtrahend);
+
+	/** Convert magnitude values to decibel values in place. db = 20 * log10(val) */
+	SIGNALPROCESSING_API void ArrayMagnitudeToDecibelInPlace(TArray<float>& InArray);
+
+	/** Convert power values to decibel values in place. db = 10 * log10(val) */
+	SIGNALPROCESSING_API void ArrayPowerToDecibelInPlace(TArray<float>& InArray);
+
+	/** FContiguousSparse2DKernelTransform
+	 *
+	 *  FContiguousSparse2DKernelTransform applies a matrix transformation to an input array. 
+	 *  [OutArray] = [[Kernal]][InArray]  
+	 *
+	 *  It provides some optimization by exploit the contiguous and sparse qualities of the kernel rows,
+	 *  which allows it to skip multiplications with the number zero. 
+	 *
+	 *  It works with non-sparse and non-contiguous kernels as well, but will be more computationally 
+	 *  expensive than a naive implementation. Also, only takes advantage of sparse contiguous rows, not columns.
+	 */
+	class SIGNALPROCESSING_API FContiguousSparse2DKernelTransform
+	{
+		struct FRow
+		{
+			int32 StartIndex;
+			TArray<float> OffsetValues;
+		};
+
+	public:
+
+		/**
+		 * NumInElements sets the expected number of input array elements as well as the number of elements in a row.
+		 * NumOutElements sets the number of output array elements as well as the number or rows.
+		 */
+		FContiguousSparse2DKernelTransform(const int32 NumInElements, const int32 NumOutElements);
+
+		/** Returns the required size of the input array */
+		int32 GetNumInElements() const;
+
+		/** Returns the size of the output array */
+		int32 GetNumOutElements() const;
+	
+		/** Set the kernel values for an individual row.
+		 *
+		 *  RowIndex determines which row is being set.
+		 *  StartIndex denotes the offset into the row where the OffsetValues will be inserted.
+		 *  OffsetValues contains the contiguous chunk of values which represent all the nonzero elements in the row.
+		 */
+		void SetRow(const int32 RowIndex, const int32 StartIndex, TArrayView<const float> OffsetValues);
+
+		/** Transforms the input array given the kernel.
+		 *
+		 *  InArray is the array to be transformed. It must have `NumInElements` number of elements.
+		 *  OutArray is the transformed array. It will have `NumOutElements` number of elements.
+		 */
+		void TransformArray(TArrayView<const float> InArray, TArray<float>& OutArray) const;
+
+		/** Transforms the input array given the kernel.
+		 *
+		 *  InArray is the array to be transformed. It must have `NumInElements` number of elements.
+		 *  OutArray is the transformed array. It will have `NumOutElements` number of elements.
+		 */
+		void TransformArray(TArrayView<const float> InArray, AlignedFloatBuffer& OutArray) const;
+
+		/** Transforms the input array given the kernel.
+		 *
+		 *  InArray is the array to be transformed. It must have `NumInElements` number of elements.
+		 *  OutArray is the transformed array. It must be allocated to hold at least NumOutElements. 
+		 */
+		void TransformArray(const float* InArray, float* OutArray) const;
+
+	private:
+
+		int32 NumIn;
+		int32 NumOut;
+		TArray<FRow> Kernel;
+	};
+
+}
diff --git a/Engine/Source/Runtime/SignalProcessing/Public/DSP/SlidingWindow.h b/Engine/Source/Runtime/SignalProcessing/Public/DSP/SlidingWindow.h
new file mode 100644
--- /dev/null
+++ b/Engine/Source/Runtime/SignalProcessing/Public/DSP/SlidingWindow.h
@@ -0,0 +1,487 @@
+// Copyright 1998-2019 Epic Games, Inc. All Rights Reserved.
+
+/** Sliding Window implementation which enables ranged for loop iteration over 
+ * sequential input buffers of varying length.
+ */
+#pragma once
+
+#include "CoreMinimal.h"
+
+
+namespace Audio
+{
+	// Forward delcaration
+	template <typename InSampleType>
+	class TSlidingWindow;
+
+
+	/** TSlidingBuffer
+	 *
+	 * TSlidingBuffer defines the window size and hop size of the sliding window, and
+	 * it stores any samples needed to produce additional windows.
+	 *
+	 * TSlidingBuffer should be used in conjunction with the TSlidingWindow, TScopedSlidingWindow 
+	 * or TAutoSlidingWindow classes.
+	 */
+	template <typename InSampleType>
+	class TSlidingBuffer
+	{
+			// Give TSlidingWindow access to StorageBuffer without exposing StorageBuffer public interface.
+			friend class TSlidingWindow<InSampleType>;
+
+
+		public:
+			/** NumWindowSamples describes the number of samples in a window */
+			const int32 NumWindowSamples;
+
+			/** NumHopSamples describes the number of samples between adjacent windows */
+			const int32 NumHopSamples;
+
+			/**
+			 * Constructs a TSlidingBuffer with a constant window and hop size
+			 */
+			TSlidingBuffer(const int32 InNumWindowSamples, const int32 InNumHopSamples)
+			: NumWindowSamples(InNumWindowSamples)
+			, NumHopSamples(InNumHopSamples)
+			, NumUnderflowSamples(0)
+			{
+				check(NumWindowSamples >= 1);
+				check(NumHopSamples >= 1);
+			}
+
+			/**
+			 * StoreForFutureWindows stores the necessary samples from InBuffer which will
+			 * be needed for future windows. It ignores all values in InBuffer which can
+			 * already be composed as a complete window.
+			 */
+			void StoreForFutureWindows(TArrayView<const InSampleType> InBuffer)
+			{
+				if (NumUnderflowSamples > 0)
+				{
+					// Consume some underflow samples from the storage buffer.
+					if (StorageBuffer.Num() > NumUnderflowSamples)
+					{
+						StorageBuffer.RemoveAt(0, NumUnderflowSamples);
+						NumUnderflowSamples = 0;
+					}
+					else
+					{
+						NumUnderflowSamples -= StorageBuffer.Num();
+						StorageBuffer.Reset();
+					}
+				}
+
+				// Total number of samples starting at beginning of first window generated from this buffer
+				int32 NumSamples = InBuffer.Num() + StorageBuffer.Num() - NumUnderflowSamples;
+
+				if (NumSamples < 0)
+				{
+					// No windows generated, but some underflow samples accounted for
+					NumUnderflowSamples -= InBuffer.Num();
+				}
+				else if (NumSamples < NumWindowSamples)
+				{
+					// No windows generated, but we should store data for future windows.
+					int32 NumToCopy = NumSamples - StorageBuffer.Num();
+					int32 InBufferIndex = InBuffer.Num() - NumToCopy;
+					StorageBuffer.Append(&InBuffer.GetData()[InBufferIndex], NumToCopy);
+
+					// All underflow samples accounted for.
+					NumUnderflowSamples = 0;
+				}
+				else
+				{
+					// Calculate number of windows generated from samples in StorageBuffer and InBuffer.
+					int32 NumWindowsGenerated = (NumSamples - NumWindowSamples) / NumHopSamples + 1;
+
+					// Calculate number of samples to keep for future windows
+					int32 NumRemainingSamples = NumSamples - (NumWindowsGenerated * NumHopSamples);
+
+					if (NumRemainingSamples > InBuffer.Num())
+					{
+						// Need to keep some samples from storage buffer and from InBuffer
+						int32 NumToKeep = NumRemainingSamples - InBuffer.Num();
+						int32 NumToRemove = StorageBuffer.Num() - NumToKeep;
+
+						if (NumToRemove > 0)
+						{
+							// May need to remove soem from the storage buffer
+							StorageBuffer.RemoveAt(0, NumToRemove);
+						}
+
+						StorageBuffer.Append(InBuffer.GetData(), InBuffer.Num());
+					}
+					else if (NumRemainingSamples > 0)
+					{
+						// Only need to keep samples from InBuffer. Can discard samples in StorageBuffer
+						StorageBuffer.Reset(NumRemainingSamples);
+						StorageBuffer.AddUninitialized(NumRemainingSamples);
+
+						int32 NewBufferCopyIndex = InBuffer.Num() - NumRemainingSamples;
+						FMemory::Memcpy(StorageBuffer.GetData(), &InBuffer.GetData()[NewBufferCopyIndex], NumRemainingSamples * sizeof(InSampleType));
+					}
+					else
+					{
+						// This occurs when HopSize > WindowSize. We have hopped to the next window, but don't have enough samples
+						// to account for the hop. We track the number of underflow samples to make sure they are consumed before
+						// the next window starts.
+						NumUnderflowSamples = -NumRemainingSamples;
+						StorageBuffer.Reset(0);
+					}
+				}
+			}
+
+			/**
+			 * Resets the internal storage.
+			 */
+			void Reset()
+			{
+				StorageBuffer.Reset();
+				NumUnderflowSamples = 0;
+			}
+
+		private:
+
+			// Stores samples from previous calls which are still needed for future buffers
+			TArray<InSampleType> StorageBuffer;
+
+			// When HopSize is greater than WindowSize a situation can occur where we need 
+			// to account for hop samples that we have not yet ingested.  
+			int32 NumUnderflowSamples;
+	};
+
+	/** TSlidingWindow
+	 *
+	 * TSlidingWindow allows windows of samples to be iterated over with STL like iterators. 
+	 *
+	 */
+	template <typename InSampleType>
+	class TSlidingWindow
+	{
+		friend class TSlidingWindowIterator;
+
+		protected:
+			// Accessed from friendship with TSlidingBuffer
+			TArrayView<const InSampleType> StorageBuffer;
+
+			// New buffer passed in.
+			TArrayView<const InSampleType> NewBuffer;
+
+			// Copied from TSlidingBuffer
+			const int32 NumWindowSamples;
+
+			// Copied from TSlidingBuffer
+			const int32 NumHopSamples;
+
+		private:
+
+			int32 MaxReadIndex;
+			int32 NumUnderflowSamples;
+
+		public:
+
+			/** TSlidingWindowIterator
+			 *
+			 * An forward iterator which slides a window over the given buffers.
+			 */
+			template<typename InAllocator = FDefaultAllocator>
+			class TSlidingWindowIterator
+			{
+					const TSlidingWindow SlidingWindow;
+
+					// Samples in window will be copied into this array.
+					TArray<InSampleType, InAllocator>& WindowBuffer;
+
+					// Index into array for reading out data.
+					int32 ReadIndex;
+
+				public:
+
+					// Sentinel value marking that the last possible window has been generated.
+					static const int32 ReadIndexEnd = INDEX_NONE;
+
+					/**
+					 * Construct an iterator over a sliding window.
+					 */
+					TSlidingWindowIterator(const TSlidingWindow& InSlidingWindow, TArray<InSampleType, InAllocator>& OutWindowBuffer, int32 InReadIndex)
+					:	SlidingWindow(InSlidingWindow)
+					,	WindowBuffer(OutWindowBuffer)
+					,	ReadIndex(InReadIndex)
+					{
+						if (ReadIndex > SlidingWindow.MaxReadIndex)
+						{
+							ReadIndex = ReadIndexEnd;
+						}
+					}
+
+					/**
+					 * Increment sliding window iterator forward.
+					 */
+					TSlidingWindowIterator operator++()
+					{
+						if (ReadIndex != ReadIndexEnd)
+						{
+							ReadIndex += SlidingWindow.NumHopSamples;
+							if (ReadIndex > SlidingWindow.MaxReadIndex)
+							{
+								ReadIndex = ReadIndexEnd;
+							}
+						}
+
+						return *this;
+					}
+
+					/**
+					 * Check whether iterators are equal. TSlidingWindowIterators derived from different
+					 * TSlidingWindows should not be compared.
+					 */
+					bool operator!=(const TSlidingWindowIterator& Other) const
+					{
+						return ReadIndex !=	Other.ReadIndex;
+					}
+
+					/**
+					 * Access array of windowed data currently pointed to by iterator.
+					 */
+					TArray<InSampleType, InAllocator>& operator*()
+					{
+						if (ReadIndex != ReadIndexEnd)
+						{
+							// Resize output window
+							WindowBuffer.Reset(SlidingWindow.NumWindowSamples);
+							WindowBuffer.AddUninitialized(SlidingWindow.NumWindowSamples);
+
+							int32 SamplesFilled = 0;
+
+							if (ReadIndex < SlidingWindow.StorageBuffer.Num())
+							{
+								// The output window overlaps the storage buffer. Copy appropriate samples from the storage buffer.
+								int32 SamplesToCopy = SlidingWindow.StorageBuffer.Num() - ReadIndex;
+								FMemory::Memcpy(WindowBuffer.GetData(), &SlidingWindow.StorageBuffer.GetData()[ReadIndex], SamplesToCopy * sizeof(InSampleType));
+								SamplesFilled += SamplesToCopy;
+							}
+
+							if (SamplesFilled < SlidingWindow.NumWindowSamples)
+							{
+								// The output window overlaps the new buffer. Copy appropriate samples from the new buffer.
+								int32 NewBufferIndex = ReadIndex - SlidingWindow.StorageBuffer.Num() + SamplesFilled;
+								int32 NewBufferRemaining = SlidingWindow.NewBuffer.Num() - NewBufferIndex;
+								int32 SamplesToCopy = FMath::Min(SlidingWindow.NumWindowSamples - SamplesFilled, NewBufferRemaining);
+								
+								if (SamplesToCopy > 0)
+								{
+									FMemory::Memcpy(&WindowBuffer.GetData()[SamplesFilled], &SlidingWindow.NewBuffer.GetData()[NewBufferIndex], SamplesToCopy * sizeof(InSampleType));
+
+									SamplesFilled += SamplesToCopy;
+								}
+							}
+
+							if (SamplesFilled < SlidingWindow.NumWindowSamples)
+							{
+								// The output window still needs more samples (due to zeropadding & flushing), so set zeros.
+								int32 SamplesToZeropad = SlidingWindow.NumWindowSamples - SamplesFilled;
+								FMemory::Memset(&WindowBuffer.GetData()[SamplesFilled], 0, sizeof(InSampleType) * SamplesToZeropad);
+							}
+						}
+						else
+						{
+							// Empty window if past end sliding window. ReadIndex == ReadIndexEnd
+							WindowBuffer.Reset();
+						}
+
+						return WindowBuffer;
+					}
+			};
+
+			/**
+			 * TSlidingWindow constructor
+			 *
+			 * InSlidingBuffer Holds the previous samples which were not completely used in previous sliding windows.  It also defines the window and hop size.
+			 * InNewBuffer Holds new samples which have not yet been ingested by the InSlidingBuffer.
+			 * bDoFlush Controls whether zeros to the final output windows until all possible windows with data from InNewBuffer have been covered.
+			 */
+			TSlidingWindow(const TSlidingBuffer<InSampleType>& InSlidingBuffer, TArrayView<const InSampleType> InNewBuffer, bool bDoFlush)
+			:	StorageBuffer(InSlidingBuffer.StorageBuffer)
+			,	NewBuffer(InNewBuffer)
+			,	NumWindowSamples(InSlidingBuffer.NumWindowSamples)
+			,	NumHopSamples(InSlidingBuffer.NumHopSamples)
+			,	MaxReadIndex(0)
+			,	NumUnderflowSamples(InSlidingBuffer.NumUnderflowSamples)
+			{
+				// Total samples to be slid over.
+				int32 NumSamples = NewBuffer.Num() + StorageBuffer.Num();
+
+				if (bDoFlush && (NumSamples > 0))
+				{
+					// If flushing, calculate the number of samples to zeropad
+					int32 NumZeroPad = 0;
+
+					if (NumSamples < NumWindowSamples)
+					{
+						NumZeroPad = NumWindowSamples - NumSamples;
+					}
+					else
+					{
+						// Determine number of windows 
+						int32 NumWindowsGenerated = (NumSamples - NumWindowSamples - NumUnderflowSamples) / NumHopSamples + 1;
+						int32 NumRemaining = NumSamples - (NumWindowsGenerated * NumHopSamples);
+						if (NumRemaining > 0)
+						{
+							NumZeroPad = NumWindowSamples - NumRemaining;
+						}
+					}
+
+					NumSamples += NumZeroPad;
+				}
+
+				MaxReadIndex = NumSamples - NumWindowSamples;
+
+				if (MaxReadIndex < 0)
+				{
+					MaxReadIndex = TSlidingWindowIterator<>::ReadIndexEnd;
+				}
+			}
+
+			virtual ~TSlidingWindow()
+			{
+			}
+
+			/**
+			 * Creates STL like iterator which slides over samples.
+			 *
+			 * OutWindowBuffer Used to construct the TSlidingWindowIterator. The iterator will populate the window with samples when the * operator is called.
+			 */
+			template<typename InAllocator = FDefaultAllocator>
+			TSlidingWindowIterator<InAllocator> begin(TArray<InSampleType, InAllocator>& OutWindowBuffer) const
+			{
+				if (MaxReadIndex == TSlidingWindowIterator<>::ReadIndexEnd)
+				{
+					return end<InAllocator>(OutWindowBuffer);
+				}
+				// Set the starting read index to NumUnderflowSamples to account for samples that still need to be consumed.
+				return TSlidingWindowIterator<InAllocator>(*this, OutWindowBuffer, NumUnderflowSamples);
+			}
+
+			/**
+			 * Creates STL like iterator denotes the end of the sliding window.
+			 *
+			 * OutWindowBuffer Used to construct the TSlidingWindowIterator. The iterator will populate the window with samples when the * operator is called.
+			 */
+			template<typename InAllocator = FDefaultAllocator>
+			TSlidingWindowIterator<InAllocator> end(TArray<InSampleType, InAllocator>& OutWindowBuffer) const
+			{
+				return TSlidingWindowIterator<InAllocator>(*this, OutWindowBuffer, TSlidingWindowIterator<InAllocator>::ReadIndexEnd);
+			}
+	};
+
+	/** TScopedSlidingWindow
+	 * 
+	 * TScopedSlidingWindow provides a sliding window iterator interface over arrays. When TScopedSlidingWindow is destructed,
+	 * it calls StoreForFutureWindow(...) on the TSlidingBuffer passed into the constructor.
+	 *
+	 */
+	template <class InSampleType>
+	class TScopedSlidingWindow : public TSlidingWindow<InSampleType>
+	{
+			// Do not allow copying or moving since that may cause the destructor to be called inadvertently.
+			TScopedSlidingWindow(TScopedSlidingWindow const &) = delete;
+			void operator=(TScopedSlidingWindow const &) = delete;
+			TScopedSlidingWindow(TScopedSlidingWindow&&) = delete;
+			TScopedSlidingWindow& operator=(TScopedSlidingWindow&&) = delete;
+
+			TSlidingBuffer<InSampleType>& SlidingBuffer;
+		public:
+
+			/**
+			 * TScopedSlidingWindow constructor
+			 *
+			 * InSlidingBuffer Holds the previous samples which were not completely used in previous sliding windows.  It also defines the window and hop size.
+			 * InNewBuffer Holds new samples which have not yet been ingested by the InSlidingBuffer.
+			 * bDoFlush Controls whether zeros to the final output windows until all possible windows with data from InNewBuffer have been covered.
+			 */
+			TScopedSlidingWindow(TSlidingBuffer<InSampleType>& InSlidingBuffer, TArrayView<const InSampleType> InNewBuffer, bool bDoFlush = false)
+			:	TSlidingWindow(InSlidingBuffer, InNewBuffer, bDoFlush)
+			,	SlidingBuffer(InSlidingBuffer)
+			{}
+
+			/**
+			 * Calls InSlidingBuffer.StoreForFutureWindows(InNewBuffer).
+			 */
+			virtual ~TScopedSlidingWindow()
+			{
+				SlidingBuffer.StoreForFutureWindows(NewBuffer);
+			}
+	};
+
+	/** TAutoSlidingWindow
+	 * 
+	 * TAutoSlidingWindow enables use of a sliding window within a range-based for loop.
+	 *
+	 * Example:
+	 *
+	 * void ProcessAudio(TSlidingBuffer<float>& SlidingBuffer, const TArray<float>& NewSamples)
+	 * {
+	 * 		TArray<float> WindowData;
+	 * 		TAutoSlidingWindow<float> SlidingWindow(SlidingBuffer, NewSamples, WindowData);
+	 *
+	 * 		for (TArray<float>& Window : SlidingWindow)
+	 * 		{
+	 * 			... audio processing on single window here
+	 * 		}
+	 * }
+	 *
+	 * int main()
+	 * {
+	 * 		int32 NumWindowSamples = 4;
+	 * 		int32 NumHopSamples = 2;
+	 * 		TSlidingBuffer<float> SlidingBuffer(NumWindowSamples, NumHopSamples);
+	 *
+	 * 		TArray<float> Buffer1({1, 2, 3, 4, 5, 6, 7});
+	 *
+	 * 		ProcessAudio(SlidingBuffer, Buffer1);
+	 *
+	 * 		TArray<float> Buffer2({8, 9, 10, 11});
+	 *
+	 * 		ProcessAudio(SlidingBuffer, Buffer2);
+	 * }
+	 */
+	template <typename InSampleType, typename InAllocator = FDefaultAllocator>
+	class TAutoSlidingWindow : public TScopedSlidingWindow<InSampleType>
+	{
+		TArray<InSampleType, InAllocator>& WindowBuffer;
+
+		public:
+			/**
+			 * TAutoSlidingWindow constructor
+			 *
+			 * InSlidingBuffer Holds the previous samples which were not completely used in previous sliding windows.  It also defines the window and hop size.
+			 * InNewBuffer Holds new samples which have not yet been ingested by the InSlidingBuffer.
+			 * OutWindow is shared by all iterators created by calling begin() or end().
+			 * bDoFlush Controls whether zeros to the final output windows until all possible windows with data from InNewBuffer have been covered.
+			 */
+			TAutoSlidingWindow(TSlidingBuffer<InSampleType>& InBuffer, TArrayView<const InSampleType> InNewBuffer, TArray<InSampleType, InAllocator>& OutWindow, bool bDoFlush = false)
+			:	TScopedSlidingWindow(InBuffer, InNewBuffer, bDoFlush)
+			,	WindowBuffer(OutWindow)
+			{}
+
+			/**
+			 * Creates STL like iterator which slides over samples.
+			 *
+			 * This iterator maintains a reference to the OutWindow passed into the constructor. That array will be manipulated when the iterator's * operator is called. 
+			 */
+			TSlidingWindow<InSampleType>::TSlidingWindowIterator<InAllocator> begin() 
+			{
+				return TSlidingWindow<InSampleType>::begin<InAllocator>(WindowBuffer);
+			}
+
+			/**
+			 * Creates STL like iterator denotes the end of the sliding window.
+			 *
+			 * This iterator maintains a reference to the OutWindow passed into the constructor. That array will be manipulated when the iterator's * operator is called. 
+			 */
+			TSlidingWindow<InSampleType>::TSlidingWindowIterator<InAllocator> end()
+			{
+				return TSlidingWindow<InSampleType>::end<InAllocator>(WindowBuffer);
+			}
+	};
+}
