Switching to TArrayViews in some instances
Adding magnitude and power spectrum methods.
Adding some array math routines.

diff --git a/Engine/Source/Runtime/SignalProcessing/Public/DSP/BufferVectorOperations.h b/Engine/Source/Runtime/SignalProcessing/Public/DSP/BufferVectorOperations.h
--- a/Engine/Source/Runtime/SignalProcessing/Public/DSP/BufferVectorOperations.h
+++ b/Engine/Source/Runtime/SignalProcessing/Public/DSP/BufferVectorOperations.h
@@ -27,90 +27,79 @@ namespace Audio
 	/** CHANNEL-AGNOSTIC OPERATIONS */
 
 	/* Sets a values to zero if value is denormal. Denormal numbers significantly slow down floating point operations. */
-	SIGNALPROCESSING_API void BufferUnderflowClampFast(AlignedFloatBuffer& InOutBuffer);
+	void AUDIOMIXER_API BufferUnderflowClampFast(AlignedFloatBuffer& InOutBuffer);
 
 	/* Sets a values to zero if value is denormal. Denormal numbers significantly slow down floating point operations. */
-	SIGNALPROCESSING_API void BufferUnderflowClampFast(float* RESTRICT InOutBuffer, const int32 InNum);
+	void AUDIOMIXER_API BufferUnderflowClampFast(float* RESTRICT InOutBuffer, const int32 InNum);
 		
 	/** Multiplies the input aligned float buffer with the given value. */
-	SIGNALPROCESSING_API void BufferMultiplyByConstant(const AlignedFloatBuffer& InFloatBuffer, float InValue, AlignedFloatBuffer& OutFloatBuffer);
-	SIGNALPROCESSING_API void BufferMultiplyByConstant(const float* RESTRICT InFloatBuffer, float InValue, float* RESTRICT OutFloatBuffer, const int32 InNumSamples);
-	SIGNALPROCESSING_API void BufferMultiplyByConstant(const AlignedFloatBuffer& InFloatBuffer, float InValue, AlignedFloatBuffer& OutFloatBuffer);
+	void AUDIOMIXER_API BufferMultiplyByConstant(const AlignedFloatBuffer& InFloatBuffer, float InValue, AlignedFloatBuffer& OutFloatBuffer);
+	void AUDIOMIXER_API BufferMultiplyByConstant(const float* RESTRICT InFloatBuffer, float InValue, float* RESTRICT OutFloatBuffer, const int32 InNumSamples);
 
 	/** Similar to BufferMultiplyByConstant, but (a) assumes a buffer length divisible by 4 and (b) performs the multiply in place. */
-	SIGNALPROCESSING_API void MultiplyBufferByConstantInPlace(AlignedFloatBuffer& InBuffer, float InGain);
-	SIGNALPROCESSING_API void MultiplyBufferByConstantInPlace(float* RESTRICT InBuffer, int32 NumSamples, float InGain);
-
-	/** Adds a constant to a buffer (useful for DC offset removal) */
-	SIGNALPROCESSING_API void AddConstantToBufferInplace(AlignedFloatBuffer& InBuffer, float Constant);
-	SIGNALPROCESSING_API void AddConstantToBufferInplace(float* RESTRICT InBuffer, int32 NumSamples, float Constant);
+	void AUDIOMIXER_API MultiplyBufferByConstantInPlace(AlignedFloatBuffer& InBuffer, float InGain);
+	void AUDIOMIXER_API MultiplyBufferByConstantInPlace(float* RESTRICT InBuffer, int32 NumSamples, float InGain);
 
 	/* Performs an element-wise weighted sum OutputBuffer = (InBuffer1 x InGain1) + (InBuffer2 x InGain2) */
-	SIGNALPROCESSING_API void BufferWeightedSumFast(const AlignedFloatBuffer& InBuffer1, float InGain1, const AlignedFloatBuffer& InBuffer2, float InGain2, AlignedFloatBuffer& OutBuffer);
+	void AUDIOMIXER_API BufferWeightedSumFast(const AlignedFloatBuffer& InBuffer1, float InGain1, const AlignedFloatBuffer& InBuffer2, float InGain2, AlignedFloatBuffer& OutBuffer);
 
 	/* Performs an element-wise weighted sum OutputBuffer = (InBuffer1 x InGain1) + InBuffer2 */
-	SIGNALPROCESSING_API void BufferWeightedSumFast(const AlignedFloatBuffer& InBuffer1, float InGain1, const AlignedFloatBuffer& InBuffer2, AlignedFloatBuffer& OutBuffer);
+	void AUDIOMIXER_API BufferWeightedSumFast(const AlignedFloatBuffer& InBuffer1, float InGain1, const AlignedFloatBuffer& InBuffer2, AlignedFloatBuffer& OutBuffer);
 
 	/* Performs an element-wise weighted sum OutputBuffer = (InBuffer1 x InGain1) + (InBuffer2 x InGain2) */
-	SIGNALPROCESSING_API void BufferWeightedSumFast(const float* RESTRICT InBuffer1, float InGain1, const float* RESTRICT InBuffer2, float InGain2, float* RESTRICT OutBuffer, int32 InNum);
+	void AUDIOMIXER_API BufferWeightedSumFast(const float* RESTRICT InBuffer1, float InGain1, const float* RESTRICT InBuffer2, float InGain2, float* RESTRICT OutBuffer, int32 InNum);
 
 	/* Performs an element-wise weighted sum OutputBuffer = (InBuffer1 x InGain1) + InBuffer2 */
-	SIGNALPROCESSING_API void BufferWeightedSumFast(const float* RESTRICT InBuffer1, float InGain1, const float* RESTRICT InBuffer2, float* RESTRICT OutBuffer, int32 InNum);
-	SIGNALPROCESSING_API void MultiplyBufferByConstantInPlace(AlignedFloatBuffer& InBuffer, float InGain);
-	SIGNALPROCESSING_API void MultiplyBufferByConstantInPlace(float* RESTRICT InBuffer, int32 NumSamples, float InGain);
+	void AUDIOMIXER_API BufferWeightedSumFast(const float* RESTRICT InBuffer1, float InGain1, const float* RESTRICT InBuffer2, float* RESTRICT OutBuffer, int32 InNum);
 
 	/* Takes a float buffer and quickly interpolates it's gain from StartValue to EndValue. */
 	/* This operation completely ignores channel counts, so avoid using this function on buffers that are not mono, stereo or quad */
 	/* if the buffer needs to fade all channels uniformly. */
-	SIGNALPROCESSING_API void FadeBufferFast(AlignedFloatBuffer& OutFloatBuffer, const float StartValue, const float EndValue);
-	SIGNALPROCESSING_API void FadeBufferFast(float* RESTRICT OutFloatBuffer, int32 NumSamples, const float StartValue, const float EndValue);
+	void AUDIOMIXER_API FadeBufferFast(AlignedFloatBuffer& OutFloatBuffer, const float StartValue, const float EndValue);
+	void AUDIOMIXER_API FadeBufferFast(float* RESTRICT OutFloatBuffer, int32 NumSamples, const float StartValue, const float EndValue);
 
 	/** Takes buffer InFloatBuffer, optionally multiplies it by Gain, and adds it to BufferToSumTo. */
-	SIGNALPROCESSING_API void MixInBufferFast(const AlignedFloatBuffer& InFloatBuffer, AlignedFloatBuffer& BufferToSumTo, const float Gain);
-	SIGNALPROCESSING_API void MixInBufferFast(const float* RESTRICT InFloatBuffer, float* RESTRICT BufferToSumTo, int32 NumSamples, const float Gain);
-	SIGNALPROCESSING_API void MixInBufferFast(const AlignedFloatBuffer& InFloatBuffer, AlignedFloatBuffer& BufferToSumTo);
-	SIGNALPROCESSING_API void MixInBufferFast(const float* RESTRICT InFloatBuffer, float* RESTRICT BufferToSumTo, int32 NumSamples);
+	void AUDIOMIXER_API MixInBufferFast(const AlignedFloatBuffer& InFloatBuffer, AlignedFloatBuffer& BufferToSumTo, const float Gain);
+	void AUDIOMIXER_API MixInBufferFast(const float* RESTRICT InFloatBuffer, float* RESTRICT BufferToSumTo, int32 NumSamples, const float Gain);
+	void AUDIOMIXER_API MixInBufferFast(const AlignedFloatBuffer& InFloatBuffer, AlignedFloatBuffer& BufferToSumTo);
+	void AUDIOMIXER_API MixInBufferFast(const float* RESTRICT InFloatBuffer, float* RESTRICT BufferToSumTo, int32 NumSamples);
 
 	/* Subtracts two buffers together element-wise. */
-	SIGNALPROCESSING_API void BufferSubtractFast(const AlignedFloatBuffer& InMinuend, const AlignedFloatBuffer& InSubtrahend, AlignedFloatBuffer& OutputBuffer);
+	void AUDIOMIXER_API BufferSubtractFast(const AlignedFloatBuffer& InMinuend, const AlignedFloatBuffer& InSubtrahend, AlignedFloatBuffer& OutputBuffer);
 
 	/* Subtracts two buffers together element-wise. */
-	SIGNALPROCESSING_API void BufferSubtractFast(const float* RESTRICT InMinuend, const float* RESTRICT InSubtrahend, float* RESTRICT OutputBuffer, int32 NumSamples);
+	void AUDIOMIXER_API BufferSubtractFast(const float* RESTRICT InMinuend, const float* RESTRICT InSubtrahend, float* RESTRICT OutputBuffer, int32 NumSamples);
 
 	/* Performs element-wise in-place subtraction placing the result in the subtrahend. InOutSubtrahend = InMinuend - InOutSubtrahend */
-	SIGNALPROCESSING_API void BufferSubtractInPlace1Fast(const AlignedFloatBuffer& InMinuend, AlignedFloatBuffer& InOutSubtrahend);
+	void AUDIOMIXER_API BufferSubtractInPlace1Fast(const AlignedFloatBuffer& InMinuend, AlignedFloatBuffer& InOutSubtrahend);
 
 	/* Performs element-wise in-place subtraction placing the result in the subtrahend. InOutSubtrahend = InMinuend - InOutSubtrahend */
-	SIGNALPROCESSING_API void BufferSubtractInPlace1Fast(const float* RESTRICT InMinuend, float* RESTRICT InOutSubtrahend, int32 NumSamples);
+	void AUDIOMIXER_API BufferSubtractInPlace1Fast(const float* RESTRICT InMinuend, float* RESTRICT InOutSubtrahend, int32 NumSamples);
 
 
 	/* Performs element-wise in-place subtraction placing the result in the minuend. InOutMinuend = InOutMinuend - InSubtrahend */
-	SIGNALPROCESSING_API void BufferSubtractInPlace2Fast(AlignedFloatBuffer& InOutMinuend, const AlignedFloatBuffer& InSubtrahend);
+	void AUDIOMIXER_API BufferSubtractInPlace2Fast(AlignedFloatBuffer& InOutMinuend, const AlignedFloatBuffer& InSubtrahend);
 
 	/* Performs element-wise in-place subtraction placing the result in the minuend. InOutMinuend = InOutMinuend - InSubtrahend */
-	SIGNALPROCESSING_API void BufferSubtractInPlace2Fast(float* RESTRICT InOutMinuend, const float* RESTRICT InSubtrahend, int32 NumSamples);
-
-	/** This version of MixInBufferFast will fade from StartGain to EndGain. */
-	SIGNALPROCESSING_API void MixInBufferFast(const AlignedFloatBuffer& InFloatBuffer, AlignedFloatBuffer& BufferToSumTo, const float StartGain, const float EndGain);
-	SIGNALPROCESSING_API void MixInBufferFast(const float* RESTRICT InFloatBuffer, float* RESTRICT BufferToSumTo, int32 NumSamples, const float StartGain, const float EndGain);
+	void AUDIOMIXER_API BufferSubtractInPlace2Fast(float* RESTRICT InOutMinuend, const float* RESTRICT InSubtrahend, int32 NumSamples);
 
 	/** Sums two buffers together and places the result in the resulting buffer. */
-	SIGNALPROCESSING_API void SumBuffers(const AlignedFloatBuffer& InFloatBuffer1, const AlignedFloatBuffer& InFloatBuffer2, AlignedFloatBuffer& OutputBuffer);
-	SIGNALPROCESSING_API void SumBuffers(const float* RESTRICT InFloatBuffer1, const float* RESTRICT InFloatBuffer2, float* RESTRICT OutputBuffer, int32 NumSamples);
+	void AUDIOMIXER_API SumBuffers(const AlignedFloatBuffer& InFloatBuffer1, const AlignedFloatBuffer& InFloatBuffer2, AlignedFloatBuffer& OutputBuffer);
+	void AUDIOMIXER_API SumBuffers(const float* RESTRICT InFloatBuffer1, const float* RESTRICT InFloatBuffer2, float* RESTRICT OutputBuffer, int32 NumSamples);
 
 	/** Multiply the second buffer in place by the first buffer. */
-	SIGNALPROCESSING_API void MultiplyBuffersInPlace(const AlignedFloatBuffer& InFloatBuffer, AlignedFloatBuffer& BufferToMultiply);
-	SIGNALPROCESSING_API void MultiplyBuffersInPlace(const float* RESTRICT InFloatBuffer, float* RESTRICT BufferToMultiply, int32 NumSamples);
+	void AUDIOMIXER_API MultiplyBuffersInPlace(const AlignedFloatBuffer& InFloatBuffer, AlignedFloatBuffer& BufferToMultiply);
+	void AUDIOMIXER_API MultiplyBuffersInPlace(const float* RESTRICT InFloatBuffer, float* RESTRICT BufferToMultiply, int32 NumSamples);
 
 	/** CHANNEL-AGNOSTIC ANALYSIS OPERATIONS */
 
 	/** Takes an audio buffer and returns the magnitude across that buffer. */
-	SIGNALPROCESSING_API float GetMagnitude(const AlignedFloatBuffer& Buffer);
-	SIGNALPROCESSING_API float GetMagnitude(const float* RESTRICT Buffer, int32 NumSamples);
+	float AUDIOMIXER_API GetMagnitude(const AlignedFloatBuffer& Buffer);
+	float AUDIOMIXER_API GetMagnitude(const float* RESTRICT Buffer, int32 NumSamples);
 
 	/** Takes an audio buffer and gets the average absolute amplitude across that buffer. */
-	SIGNALPROCESSING_API float GetAverageAmplitude(const AlignedFloatBuffer& Buffer);
-	SIGNALPROCESSING_API float GetAverageAmplitude(const float* RESTRICT Buffer, int32 NumSamples);
+	float AUDIOMIXER_API GetAverageAmplitude(const AlignedFloatBuffer& Buffer);
+	float AUDIOMIXER_API GetAverageAmplitude(const float* RESTRICT Buffer, int32 NumSamples);
 
 	/** CHANNEL-SPECIFIC OPERATIONS */
 
@@ -118,146 +107,147 @@ namespace Audio
 	 *  StereoBuffer must have an even number of frames.
 	 *  If StartGains and EndGains are provided, this function will interpolate between the two across the buffer.
 	*/
-	SIGNALPROCESSING_API void Apply2ChannelGain(AlignedFloatBuffer& StereoBuffer, const float* RESTRICT Gains);
-	SIGNALPROCESSING_API void Apply2ChannelGain(float* RESTRICT StereoBuffer, int32 NumSamples, const float* RESTRICT Gains);
-	SIGNALPROCESSING_API void Apply2ChannelGain(AlignedFloatBuffer& StereoBuffer, const float* RESTRICT StartGains, const float* RESTRICT EndGains);
-	SIGNALPROCESSING_API void Apply2ChannelGain(float* RESTRICT StereoBuffer, int32 NumSamples, const float* RESTRICT StartGains, const float* RESTRICT EndGains);
+	void AUDIOMIXER_API Apply2ChannelGain(AlignedFloatBuffer& StereoBuffer, const float* RESTRICT Gains);
+	void AUDIOMIXER_API Apply2ChannelGain(float* RESTRICT StereoBuffer, int32 NumSamples, const float* RESTRICT Gains);
+	void AUDIOMIXER_API Apply2ChannelGain(AlignedFloatBuffer& StereoBuffer, const float* RESTRICT StartGains, const float* RESTRICT EndGains);
+	void AUDIOMIXER_API Apply2ChannelGain(float* RESTRICT StereoBuffer, int32 NumSamples, const float* RESTRICT StartGains, const float* RESTRICT EndGains);
 
 	/** Takes a 1 channel buffer and mixes it to a stereo buffer using Gains. Gains is expected to point to a 2 float long buffer.
 	 *  these buffers must have an even number of frames.
 	 *  If StartGains and EndGains are provided, this function will interpolate between the two across the buffer.
 	*/
-	SIGNALPROCESSING_API void MixMonoTo2ChannelsFast(const AlignedFloatBuffer& MonoBuffer, AlignedFloatBuffer& DestinationBuffer, const float* RESTRICT Gains);
-	SIGNALPROCESSING_API void MixMonoTo2ChannelsFast(const float* RESTRICT MonoBuffer, float* RESTRICT DestinationBuffer, int32 NumFrames, const float* RESTRICT Gains);
-	SIGNALPROCESSING_API void MixMonoTo2ChannelsFast(const AlignedFloatBuffer& MonoBuffer, AlignedFloatBuffer& DestinationBuffer, const float* RESTRICT StartGains, const float* RESTRICT EndGains);
-	SIGNALPROCESSING_API void MixMonoTo2ChannelsFast(const float* RESTRICT MonoBuffer, float* RESTRICT DestinationBuffer, int32 NumFrames, const float* RESTRICT StartGains, const float* RESTRICT EndGains);
-	SIGNALPROCESSING_API void MixMonoTo2ChannelsFast(const AlignedFloatBuffer& MonoBuffer, AlignedFloatBuffer& DestinationBuffer);
-	SIGNALPROCESSING_API void MixMonoTo2ChannelsFast(const float* RESTRICT MonoBuffer, float* RESTRICT DestinationBuffer, int32 NumFrames);
+	void AUDIOMIXER_API MixMonoTo2ChannelsFast(const AlignedFloatBuffer& MonoBuffer, AlignedFloatBuffer& DestinationBuffer, const float* RESTRICT Gains);
+	void AUDIOMIXER_API MixMonoTo2ChannelsFast(const float* RESTRICT MonoBuffer, float* RESTRICT DestinationBuffer, int32 NumFrames, const float* RESTRICT Gains);
+	void AUDIOMIXER_API MixMonoTo2ChannelsFast(const AlignedFloatBuffer& MonoBuffer, AlignedFloatBuffer& DestinationBuffer, const float* RESTRICT StartGains, const float* RESTRICT EndGains);
+	void AUDIOMIXER_API MixMonoTo2ChannelsFast(const float* RESTRICT MonoBuffer, float* RESTRICT DestinationBuffer, int32 NumFrames, const float* RESTRICT StartGains, const float* RESTRICT EndGains);
+	void AUDIOMIXER_API MixMonoTo2ChannelsFast(const AlignedFloatBuffer& MonoBuffer, AlignedFloatBuffer& DestinationBuffer);
+	void AUDIOMIXER_API MixMonoTo2ChannelsFast(const float* RESTRICT MonoBuffer, float* RESTRICT DestinationBuffer, int32 NumFrames);
 
 	/** Takes a 2 channel buffer and mixes it to an 2 channel interleaved buffer using Gains. Gains is expected to point to a 16 float long buffer.
 	*  Output gains for the left input channel should be the first 8 values in Gains, and Output gains for the right input channel should be rest.
 	*  NumFrames must be a multiple of 4.
 	*  If StartGains and EndGains are provided, this function will interpolate between the two across the buffer.
 	*/
-	SIGNALPROCESSING_API void Mix2ChannelsTo2ChannelsFast(const AlignedFloatBuffer& SourceBuffer, AlignedFloatBuffer& DestinationBuffer, const float* RESTRICT Gains);
-	SIGNALPROCESSING_API void Mix2ChannelsTo2ChannelsFast(const float* RESTRICT SourceBuffer, float* RESTRICT DestinationBuffer, int32 NumFrames, const float* RESTRICT Gains);
-	SIGNALPROCESSING_API void Mix2ChannelsTo2ChannelsFast(const AlignedFloatBuffer& SourceBuffer, AlignedFloatBuffer& DestinationBuffer, const float* RESTRICT StartGains, const float* RESTRICT EndGains);
-	SIGNALPROCESSING_API void Mix2ChannelsTo2ChannelsFast(const float* RESTRICT SourceBuffer, float* RESTRICT DestinationBuffer, int32 NumFrames, const float* RESTRICT StartGains, const float* RESTRICT EndGains);
+	void AUDIOMIXER_API Mix2ChannelsTo2ChannelsFast(const AlignedFloatBuffer& SourceBuffer, AlignedFloatBuffer& DestinationBuffer, const float* RESTRICT Gains);
+	void AUDIOMIXER_API Mix2ChannelsTo2ChannelsFast(const float* RESTRICT SourceBuffer, float* RESTRICT DestinationBuffer, int32 NumFrames, const float* RESTRICT Gains);
+	void AUDIOMIXER_API Mix2ChannelsTo2ChannelsFast(const AlignedFloatBuffer& SourceBuffer, AlignedFloatBuffer& DestinationBuffer, const float* RESTRICT StartGains, const float* RESTRICT EndGains);
+	void AUDIOMIXER_API Mix2ChannelsTo2ChannelsFast(const float* RESTRICT SourceBuffer, float* RESTRICT DestinationBuffer, int32 NumFrames, const float* RESTRICT StartGains, const float* RESTRICT EndGains);
 
 	/** Takes a 4 channel interleaved buffer and applies Gains to it. Gains is expected to point to a 2 float long buffer.
 	 *  If StartGains and EndGains are provided, this function will interpolate between the two across the buffer.
 	*/
-	SIGNALPROCESSING_API void Apply4ChannelGain(AlignedFloatBuffer& InterleavedBuffer, const float* RESTRICT Gains);
-	SIGNALPROCESSING_API void Apply4ChannelGain(float* RESTRICT InterleavedBuffer, int32 NumSamples, const float* RESTRICT Gains);
-	SIGNALPROCESSING_API void Apply4ChannelGain(AlignedFloatBuffer& InterleavedBuffer, const float* RESTRICT StartGains, const float* RESTRICT EndGains);
-	SIGNALPROCESSING_API void Apply4ChannelGain(float* RESTRICT InterleavedBuffer, int32 NumSamples, const float* RESTRICT StartGains, const float* RESTRICT EndGains);
+	void AUDIOMIXER_API Apply4ChannelGain(AlignedFloatBuffer& InterleavedBuffer, const float* RESTRICT Gains);
+	void AUDIOMIXER_API Apply4ChannelGain(float* RESTRICT InterleavedBuffer, int32 NumSamples, const float* RESTRICT Gains);
+	void AUDIOMIXER_API Apply4ChannelGain(AlignedFloatBuffer& InterleavedBuffer, const float* RESTRICT StartGains, const float* RESTRICT EndGains);
+	void AUDIOMIXER_API Apply4ChannelGain(float* RESTRICT InterleavedBuffer, int32 NumSamples, const float* RESTRICT StartGains, const float* RESTRICT EndGains);
 
 	/** Takes a 1 channel buffer and mixes it to an 8 channel interleaved buffer using Gains. Gains is expected to point to a 8 float long buffer.
 	*  these buffers must have an even number of frames.
 	*  If StartGains and EndGains are provided, this function will interpolate between the two across the buffer.
 	*/
-	SIGNALPROCESSING_API void MixMonoTo4ChannelsFast(const AlignedFloatBuffer& MonoBuffer, AlignedFloatBuffer& DestinationBuffer, const float* RESTRICT Gains);
-	SIGNALPROCESSING_API void MixMonoTo4ChannelsFast(const float* RESTRICT MonoBuffer, float* RESTRICT DestinationBuffer, int32 NumFrames, const float* RESTRICT Gains);
-	SIGNALPROCESSING_API void MixMonoTo4ChannelsFast(const AlignedFloatBuffer& MonoBuffer, AlignedFloatBuffer& DestinationBuffer, const float* RESTRICT StartGains, const float* RESTRICT EndGains);
-	SIGNALPROCESSING_API void MixMonoTo4ChannelsFast(const float* RESTRICT MonoBuffer, float* RESTRICT DestinationBuffer, int32 NumFrames, const float* RESTRICT StartGains, const float* RESTRICT EndGains);
+	void AUDIOMIXER_API MixMonoTo4ChannelsFast(const AlignedFloatBuffer& MonoBuffer, AlignedFloatBuffer& DestinationBuffer, const float* RESTRICT Gains);
+	void AUDIOMIXER_API MixMonoTo4ChannelsFast(const float* RESTRICT MonoBuffer, float* RESTRICT DestinationBuffer, int32 NumFrames, const float* RESTRICT Gains);
+	void AUDIOMIXER_API MixMonoTo4ChannelsFast(const AlignedFloatBuffer& MonoBuffer, AlignedFloatBuffer& DestinationBuffer, const float* RESTRICT StartGains, const float* RESTRICT EndGains);
+	void AUDIOMIXER_API MixMonoTo4ChannelsFast(const float* RESTRICT MonoBuffer, float* RESTRICT DestinationBuffer, int32 NumFrames, const float* RESTRICT StartGains, const float* RESTRICT EndGains);
 
 	/** Takes a 2 channel buffer and mixes it to an 8 channel interleaved buffer using Gains. Gains is expected to point to a 16 float long buffer.
 	*  Output gains for the left input channel should be the first 8 values in Gains, and Output gains for the right input channel should be rest.
 	*  NumFrames must be a multiple of 4.
 	*  If StartGains and EndGains are provided, this function will interpolate between the two across the buffer.
 	*/
-	SIGNALPROCESSING_API void Mix2ChannelsTo4ChannelsFast(const AlignedFloatBuffer& SourceBuffer, AlignedFloatBuffer& DestinationBuffer, const float* RESTRICT Gains);
-	SIGNALPROCESSING_API void Mix2ChannelsTo4ChannelsFast(const float* RESTRICT SourceBuffer, float* RESTRICT DestinationBuffer, int32 NumFrames, const float* RESTRICT Gains);
-	SIGNALPROCESSING_API void Mix2ChannelsTo4ChannelsFast(const AlignedFloatBuffer& SourceBuffer, AlignedFloatBuffer& DestinationBuffer, const float* RESTRICT StartGains, const float* RESTRICT EndGains);
-	SIGNALPROCESSING_API void Mix2ChannelsTo4ChannelsFast(const float* RESTRICT SourceBuffer, float* RESTRICT DestinationBuffer, int32 NumFrames, const float* RESTRICT StartGains, const float* RESTRICT EndGains);
+	void AUDIOMIXER_API Mix2ChannelsTo4ChannelsFast(const AlignedFloatBuffer& SourceBuffer, AlignedFloatBuffer& DestinationBuffer, const float* RESTRICT Gains);
+	void AUDIOMIXER_API Mix2ChannelsTo4ChannelsFast(const float* RESTRICT SourceBuffer, float* RESTRICT DestinationBuffer, int32 NumFrames, const float* RESTRICT Gains);
+	void AUDIOMIXER_API Mix2ChannelsTo4ChannelsFast(const AlignedFloatBuffer& SourceBuffer, AlignedFloatBuffer& DestinationBuffer, const float* RESTRICT StartGains, const float* RESTRICT EndGains);
+	void AUDIOMIXER_API Mix2ChannelsTo4ChannelsFast(const float* RESTRICT SourceBuffer, float* RESTRICT DestinationBuffer, int32 NumFrames, const float* RESTRICT StartGains, const float* RESTRICT EndGains);
 
 	/** Takes a 6 channel interleaved buffer and applies Gains to it. Gains is expected to point to a 2 float long buffer.
 	 *  InterleavedBuffer must have an even number of frames.
 	*/
-	SIGNALPROCESSING_API void Apply6ChannelGain(AlignedFloatBuffer& InterleavedBuffer, const float* RESTRICT Gains);
-	SIGNALPROCESSING_API void Apply6ChannelGain(float* RESTRICT InterleavedBuffer, int32 NumSamples, const float* RESTRICT Gains);
-	SIGNALPROCESSING_API void Apply6ChannelGain(AlignedFloatBuffer& InterleavedBuffer, const float* RESTRICT StartGains, const float* RESTRICT EndGains);
-	SIGNALPROCESSING_API void Apply6ChannelGain(float* RESTRICT InterleavedBuffer, int32 NumSamples, const float* RESTRICT StartGains, const float* RESTRICT EndGains);
+	void AUDIOMIXER_API Apply6ChannelGain(AlignedFloatBuffer& InterleavedBuffer, const float* RESTRICT Gains);
+	void AUDIOMIXER_API Apply6ChannelGain(float* RESTRICT InterleavedBuffer, int32 NumSamples, const float* RESTRICT Gains);
+	void AUDIOMIXER_API Apply6ChannelGain(AlignedFloatBuffer& InterleavedBuffer, const float* RESTRICT StartGains, const float* RESTRICT EndGains);
+	void AUDIOMIXER_API Apply6ChannelGain(float* RESTRICT InterleavedBuffer, int32 NumSamples, const float* RESTRICT StartGains, const float* RESTRICT EndGains);
 
 	/** Takes a 1 channel buffer and mixes it to an 8 channel interleaved buffer using Gains. Gains is expected to point to a 8 float long buffer.
 	 *  these buffers must have an even number of frames.
 	 *  If StartGains and EndGains are provided, this function will interpolate between the two across the buffer.
 	*/
-	SIGNALPROCESSING_API void MixMonoTo6ChannelsFast(const AlignedFloatBuffer& MonoBuffer, AlignedFloatBuffer& DestinationBuffer, const float* RESTRICT Gains);
-	SIGNALPROCESSING_API void MixMonoTo6ChannelsFast(const float* RESTRICT MonoBuffer, float* RESTRICT DestinationBuffer, int32 NumFrames, const float* RESTRICT Gains);
-	SIGNALPROCESSING_API void MixMonoTo6ChannelsFast(const AlignedFloatBuffer& MonoBuffer, AlignedFloatBuffer& DestinationBuffer, const float* RESTRICT StartGains, const float* RESTRICT EndGains);
-	SIGNALPROCESSING_API void MixMonoTo6ChannelsFast(const float* RESTRICT MonoBuffer, float* RESTRICT DestinationBuffer, int32 NumFrames, const float* RESTRICT StartGains, const float* RESTRICT EndGains);
+	void AUDIOMIXER_API MixMonoTo6ChannelsFast(const AlignedFloatBuffer& MonoBuffer, AlignedFloatBuffer& DestinationBuffer, const float* RESTRICT Gains);
+	void AUDIOMIXER_API MixMonoTo6ChannelsFast(const float* RESTRICT MonoBuffer, float* RESTRICT DestinationBuffer, int32 NumFrames, const float* RESTRICT Gains);
+	void AUDIOMIXER_API MixMonoTo6ChannelsFast(const AlignedFloatBuffer& MonoBuffer, AlignedFloatBuffer& DestinationBuffer, const float* RESTRICT StartGains, const float* RESTRICT EndGains);
+	void AUDIOMIXER_API MixMonoTo6ChannelsFast(const float* RESTRICT MonoBuffer, float* RESTRICT DestinationBuffer, int32 NumFrames, const float* RESTRICT StartGains, const float* RESTRICT EndGains);
 
 	/** Takes a 2 channel buffer and mixes it to an 8 channel interleaved buffer using Gains. Gains is expected to point to a 16 float long buffer.
 	 *  Output gains for the left input channel should be the first 8 values in Gains, and Output gains for the right input channel should be rest.
 	 *  NumFrames must be a multiple of 4.
 	 *  If StartGains and EndGains are provided, this function will interpolate between the two across the buffer.
 	*/
-	SIGNALPROCESSING_API void Mix2ChannelsTo6ChannelsFast(const AlignedFloatBuffer& SourceBuffer, AlignedFloatBuffer& DestinationBuffer, const float* RESTRICT Gains);
-	SIGNALPROCESSING_API void Mix2ChannelsTo6ChannelsFast(const float* RESTRICT SourceBuffer, float* RESTRICT DestinationBuffer, int32 NumFrames, const float* RESTRICT Gains);
-	SIGNALPROCESSING_API void Mix2ChannelsTo6ChannelsFast(const AlignedFloatBuffer& SourceBuffer, AlignedFloatBuffer& DestinationBuffer, const float* RESTRICT StartGains, const float* RESTRICT EndGains);
-	SIGNALPROCESSING_API void Mix2ChannelsTo6ChannelsFast(const float* RESTRICT SourceBuffer, float* RESTRICT DestinationBuffer, int32 NumFrames, const float* RESTRICT StartGains, const float* RESTRICT EndGains);
+	void AUDIOMIXER_API Mix2ChannelsTo6ChannelsFast(const AlignedFloatBuffer& SourceBuffer, AlignedFloatBuffer& DestinationBuffer, const float* RESTRICT Gains);
+	void AUDIOMIXER_API Mix2ChannelsTo6ChannelsFast(const float* RESTRICT SourceBuffer, float* RESTRICT DestinationBuffer, int32 NumFrames, const float* RESTRICT Gains);
+	void AUDIOMIXER_API Mix2ChannelsTo6ChannelsFast(const AlignedFloatBuffer& SourceBuffer, AlignedFloatBuffer& DestinationBuffer, const float* RESTRICT StartGains, const float* RESTRICT EndGains);
+	void AUDIOMIXER_API Mix2ChannelsTo6ChannelsFast(const float* RESTRICT SourceBuffer, float* RESTRICT DestinationBuffer, int32 NumFrames, const float* RESTRICT StartGains, const float* RESTRICT EndGains);
 
 	/** Takes an 8 channel interleaved buffer and applies Gains to it. Gains is expected to point to an 8 float long buffer. */
-	SIGNALPROCESSING_API void Apply8ChannelGain(AlignedFloatBuffer& InterleavedBuffer, const float* RESTRICT Gains);
-	SIGNALPROCESSING_API void Apply8ChannelGain(float* RESTRICT InterleavedBuffer, int32 NumSamples, const float* RESTRICT Gains);
-	SIGNALPROCESSING_API void Apply8ChannelGain(AlignedFloatBuffer& InterleavedBuffer, const float* RESTRICT StartGains, const float* RESTRICT EndGains);
-	SIGNALPROCESSING_API void Apply8ChannelGain(float* RESTRICT InterleavedBuffer, int32 NumSamples, const float* RESTRICT StartGains, const float* RESTRICT EndGains);
+	void AUDIOMIXER_API Apply8ChannelGain(AlignedFloatBuffer& InterleavedBuffer, const float* RESTRICT Gains);
+	void AUDIOMIXER_API Apply8ChannelGain(float* RESTRICT InterleavedBuffer, int32 NumSamples, const float* RESTRICT Gains);
+	void AUDIOMIXER_API Apply8ChannelGain(AlignedFloatBuffer& InterleavedBuffer, const float* RESTRICT StartGains, const float* RESTRICT EndGains);
+	void AUDIOMIXER_API Apply8ChannelGain(float* RESTRICT InterleavedBuffer, int32 NumSamples, const float* RESTRICT StartGains, const float* RESTRICT EndGains);
 
 	/** Takes a 1 channel buffer and mixes it to an 8 channel interleaved buffer using Gains. Gains is expected to point to a 8 float long buffer.
 	 *  these buffers must have an even number of frames.
 	 *  If StartGains and EndGains are provided, this function will interpolate between the two across the buffer.
 	*/
-	SIGNALPROCESSING_API void MixMonoTo8ChannelsFast(const AlignedFloatBuffer& MonoBuffer, AlignedFloatBuffer& DestinationBuffer, const float* RESTRICT Gains);
-	SIGNALPROCESSING_API void MixMonoTo8ChannelsFast(const float* RESTRICT MonoBuffer, float* RESTRICT DestinationBuffer, int32 NumFrames, const float* RESTRICT Gains);
-	SIGNALPROCESSING_API void MixMonoTo8ChannelsFast(const AlignedFloatBuffer& MonoBuffer, AlignedFloatBuffer& DestinationBuffer, const float* RESTRICT StartGains, const float* RESTRICT EndGains);
-	SIGNALPROCESSING_API void MixMonoTo8ChannelsFast(const float* RESTRICT MonoBuffer, float* RESTRICT DestinationBuffer, int32 NumFrames, const float* RESTRICT StartGains, const float* RESTRICT EndGains);
+	void AUDIOMIXER_API MixMonoTo8ChannelsFast(const AlignedFloatBuffer& MonoBuffer, AlignedFloatBuffer& DestinationBuffer, const float* RESTRICT Gains);
+	void AUDIOMIXER_API MixMonoTo8ChannelsFast(const float* RESTRICT MonoBuffer, float* RESTRICT DestinationBuffer, int32 NumFrames, const float* RESTRICT Gains);
+	void AUDIOMIXER_API MixMonoTo8ChannelsFast(const AlignedFloatBuffer& MonoBuffer, AlignedFloatBuffer& DestinationBuffer, const float* RESTRICT StartGains, const float* RESTRICT EndGains);
+	void AUDIOMIXER_API MixMonoTo8ChannelsFast(const float* RESTRICT MonoBuffer, float* RESTRICT DestinationBuffer, int32 NumFrames, const float* RESTRICT StartGains, const float* RESTRICT EndGains);
 
 	/** Takes a 2 channel buffer and mixes it to an 8 channel interleaved buffer using Gains. Gains is expected to point to a 16 float long buffer.
 	 *  Output gains for the left input channel should be the first 8 values in Gains, and Output gains for the right input channel should be rest.
 	 *  these buffers must have an even number of frames.
 	 *  If StartGains and EndGains are provided, this function will interpolate between the two across the buffer.
 	*/
-	SIGNALPROCESSING_API void Mix2ChannelsTo8ChannelsFast(const AlignedFloatBuffer& SourceBuffer, AlignedFloatBuffer& DestinationBuffer, const float* RESTRICT Gains);
-	SIGNALPROCESSING_API void Mix2ChannelsTo8ChannelsFast(const float* RESTRICT SourceBuffer, float* RESTRICT DestinationBuffer, int32 NumFrames, const float* RESTRICT Gains);
-	SIGNALPROCESSING_API void Mix2ChannelsTo8ChannelsFast(const AlignedFloatBuffer& SourceBuffer, AlignedFloatBuffer& DestinationBuffer, const float* RESTRICT StartGains, const float* RESTRICT EndGains);
-	SIGNALPROCESSING_API void Mix2ChannelsTo8ChannelsFast(const float* RESTRICT SourceBuffer, float* RESTRICT DestinationBuffer, int32 NumFrames, const float* RESTRICT StartGains, const float* RESTRICT EndGains);
+	void AUDIOMIXER_API Mix2ChannelsTo8ChannelsFast(const AlignedFloatBuffer& SourceBuffer, AlignedFloatBuffer& DestinationBuffer, const float* RESTRICT Gains);
+	void AUDIOMIXER_API Mix2ChannelsTo8ChannelsFast(const float* RESTRICT SourceBuffer, float* RESTRICT DestinationBuffer, int32 NumFrames, const float* RESTRICT Gains);
+	void AUDIOMIXER_API Mix2ChannelsTo8ChannelsFast(const AlignedFloatBuffer& SourceBuffer, AlignedFloatBuffer& DestinationBuffer, const float* RESTRICT StartGains, const float* RESTRICT EndGains);
+	void AUDIOMIXER_API Mix2ChannelsTo8ChannelsFast(const float* RESTRICT SourceBuffer, float* RESTRICT DestinationBuffer, int32 NumFrames, const float* RESTRICT StartGains, const float* RESTRICT EndGains);
 
 	/** This is a generalized operation that uses the channel gain matrix provided in Gains to mix an interleaved source buffer to the interleaved downmix buffer.
 	 *  This operation is not explicitly vectorized and will almost always be slower than using one of the functions above.
 	*/
-	SIGNALPROCESSING_API void DownmixBuffer(int32 NumSourceChannels, int32 NumDestinationChannels, const AlignedFloatBuffer& SourceBuffer, AlignedFloatBuffer& DestinationBuffer, const float* RESTRICT Gains);
-	SIGNALPROCESSING_API void DownmixBuffer(int32 NumSourceChannels, int32 NumDestinationChannels, const float* RESTRICT SourceBuffer, float* RESTRICT DestinationBuffer, int32 NumFrames, const float* RESTRICT Gains);
-	SIGNALPROCESSING_API void DownmixBuffer(int32 NumSourceChannels, int32 NumDestinationChannels, const AlignedFloatBuffer& SourceBuffer, AlignedFloatBuffer& DestinationBuffer, float* RESTRICT StartGains, const float* RESTRICT EndGains);
-	SIGNALPROCESSING_API void DownmixBuffer(int32 NumSourceChannels, int32 NumDestinationChannels, const float* RESTRICT SourceBuffer, float* RESTRICT DestinationBuffer, int32 NumFrames, float* RESTRICT StartGains, const float* RESTRICT EndGains);
+	void AUDIOMIXER_API DownmixBuffer(int32 NumSourceChannels, int32 NumDestinationChannels, const AlignedFloatBuffer& SourceBuffer, AlignedFloatBuffer& DestinationBuffer, const float* RESTRICT Gains);
+	void AUDIOMIXER_API DownmixBuffer(int32 NumSourceChannels, int32 NumDestinationChannels, const float* RESTRICT SourceBuffer, float* RESTRICT DestinationBuffer, int32 NumFrames, const float* RESTRICT Gains);
+	void AUDIOMIXER_API DownmixBuffer(int32 NumSourceChannels, int32 NumDestinationChannels, const AlignedFloatBuffer& SourceBuffer, AlignedFloatBuffer& DestinationBuffer, float* RESTRICT StartGains, const float* RESTRICT EndGains);
+	void AUDIOMIXER_API DownmixBuffer(int32 NumSourceChannels, int32 NumDestinationChannels, const float* RESTRICT SourceBuffer, float* RESTRICT DestinationBuffer, int32 NumFrames, float* RESTRICT StartGains, const float* RESTRICT EndGains);
+
 
 	/** Interleaves samples from two input buffers */
-	SIGNALPROCESSING_API void BufferInterleave2ChannelFast(const AlignedFloatBuffer& InBuffer1, const AlignedFloatBuffer& InBuffer2, AlignedFloatBuffer& OutBuffer);
+	void AUDIOMIXER_API BufferInterleave2ChannelFast(const AlignedFloatBuffer& InBuffer1, const AlignedFloatBuffer& InBuffer2, AlignedFloatBuffer& OutBuffer);
 
 	/** Interleaves samples from two input buffers */
-	SIGNALPROCESSING_API void BufferInterleave2ChannelFast(const float* RESTRICT InBuffer1, const float* RESTRICT InBuffer2, float* RESTRICT OutBuffer, const int32 InNum);
+	void AUDIOMIXER_API BufferInterleave2ChannelFast(const float* RESTRICT InBuffer1, const float* RESTRICT InBuffer2, float* RESTRICT OutBuffer, const int32 InNum);
 
 	/** Deinterleaves samples from a 2 channel input buffer */
-	SIGNALPROCESSING_API void BufferDeinterleave2ChannelFast(const AlignedFloatBuffer& InBuffer, AlignedFloatBuffer& OutBuffer1, AlignedFloatBuffer& OutBuffer2);
+	void AUDIOMIXER_API BufferDeinterleave2ChannelFast(const AlignedFloatBuffer& InBuffer, AlignedFloatBuffer& OutBuffer1, AlignedFloatBuffer& OutBuffer2);
 
 	/** Deinterleaves samples from a 2 channel input buffer */
-	SIGNALPROCESSING_API void BufferDeinterleave2ChannelFast(const float* RESTRICT InBuffer, float* RESTRICT OutBuffer1, float* RESTRICT OutBuffer2, const int32 InNum);
+	void AUDIOMIXER_API BufferDeinterleave2ChannelFast(const float* RESTRICT InBuffer, float* RESTRICT OutBuffer1, float* RESTRICT OutBuffer2, const int32 InNum);
 
 	/** Sums 2 channel interleaved input samples. OutSamples[n] = InSamples[2n] + InSamples[2n + 1] */
-	SIGNALPROCESSING_API void BufferSum2ChannelToMonoFast(const AlignedFloatBuffer& InSamples, AlignedFloatBuffer& OutSamples);
+	void AUDIOMIXER_API BufferSum2ChannelToMonoFast(const AlignedFloatBuffer& InSamples, AlignedFloatBuffer& OutSamples);
 
 	/** Sums 2 channel interleaved input samples. OutSamples[n] = InSamples[2n] + InSamples[2n + 1] */
-	void SIGNALPROCESSING_API BufferSum2ChannelToMonoFast(const float* RESTRICT InSamples, float* RESTRICT OutSamples, const int32 InNumFrames);
+	void AUDIOMIXER_API BufferSum2ChannelToMonoFast(const float* RESTRICT InSamples, float* RESTRICT OutSamples, const int32 InNumFrames);
 
 	/** Compute power of complex data. Out[i] = Real[i] * Real[i] + Imaginary[i] * Imaginary[i] */
-	void SIGNALPROCESSING_API BufferComplexToPowerFast(const AlignedFloatBuffer& InRealSamples, const AlignedFloatBuffer& InImaginarySamples, AlignedFloatBuffer& OutPowerSamples);
+	void AUDIOMIXER_API BufferComplexToPowerFast(const AlignedFloatBuffer& InRealSamples, const AlignedFloatBuffer& InImaginarySamples, AlignedFloatBuffer& OutPowerSamples);
 
 	/** Compute power of complex data. Out[i] = Real[i] * Real[i] + Imaginary[i] * Imaginary[i] */
-	void SIGNALPROCESSING_API BufferComplexToPowerFast(const float* RESTRICT InRealSamples, const float* RESTRICT InImaginarySamples, float* RESTRICT OutPowerSamples, const int32 InNum);
+	void AUDIOMIXER_API BufferComplexToPowerFast(const float* RESTRICT InRealSamples, const float* RESTRICT InImaginarySamples, float* RESTRICT OutPowerSamples, const int32 InNum);
 
 	/** Compute magnitude of complex data. Out[i] = Sqrt(Real[i] * Real[i] + Imaginary[i] * Imaginary[i]) */
-	void SIGNALPROCESSING_API BufferComplexToMagnitudeFast(const AlignedFloatBuffer& InRealSamples, const AlignedFloatBuffer& InImaginarySamples, AlignedFloatBuffer& OutPowerSamples);
+	void AUDIOMIXER_API BufferComplexToMagnitudeFast(const AlignedFloatBuffer& InRealSamples, const AlignedFloatBuffer& InImaginarySamples, AlignedFloatBuffer& OutPowerSamples);
 
 	/** Compute magnitude of complex data. Out[i] = Sqrt(Real[i] * Real[i] + Imaginary[i] * Imaginary[i]) */
-	void SIGNALPROCESSING_API BufferComplexToMagnitudeFast(const float* RESTRICT InRealSamples, const float* RESTRICT InImaginarySamples, float* RESTRICT OutPowerSamples, const int32 InNum);
+	void AUDIOMIXER_API BufferComplexToMagnitudeFast(const float* RESTRICT InRealSamples, const float* RESTRICT InImaginarySamples, float* RESTRICT OutPowerSamples, const int32 InNum);
 }
