Fix func declaration syntax in BufferVectorOperations.h #jira none #rb none

diff --git a/Engine/Source/Runtime/SignalProcessing/Public/DSP/BufferVectorOperations.h b/Engine/Source/Runtime/SignalProcessing/Public/DSP/BufferVectorOperations.h
--- a/Engine/Source/Runtime/SignalProcessing/Public/DSP/BufferVectorOperations.h
+++ b/Engine/Source/Runtime/SignalProcessing/Public/DSP/BufferVectorOperations.h
@@ -247,17 +247,17 @@ namespace Audio
 	SIGNALPROCESSING_API void BufferSum2ChannelToMonoFast(const AlignedFloatBuffer& InSamples, AlignedFloatBuffer& OutSamples);
 
 	/** Sums 2 channel interleaved input samples. OutSamples[n] = InSamples[2n] + InSamples[2n + 1] */
-	void SIGNALPROCESSING_API BufferSum2ChannelToMonoFast(const float* RESTRICT InSamples, float* RESTRICT OutSamples, const int32 InNumFrames);
+	SIGNALPROCESSING_API void BufferSum2ChannelToMonoFast(const float* RESTRICT InSamples, float* RESTRICT OutSamples, const int32 InNumFrames);
 
 	/** Compute power of complex data. Out[i] = Real[i] * Real[i] + Imaginary[i] * Imaginary[i] */
-	void SIGNALPROCESSING_API BufferComplexToPowerFast(const AlignedFloatBuffer& InRealSamples, const AlignedFloatBuffer& InImaginarySamples, AlignedFloatBuffer& OutPowerSamples);
+	SIGNALPROCESSING_API void BufferComplexToPowerFast(const AlignedFloatBuffer& InRealSamples, const AlignedFloatBuffer& InImaginarySamples, AlignedFloatBuffer& OutPowerSamples);
 
 	/** Compute power of complex data. Out[i] = Real[i] * Real[i] + Imaginary[i] * Imaginary[i] */
-	void SIGNALPROCESSING_API BufferComplexToPowerFast(const float* RESTRICT InRealSamples, const float* RESTRICT InImaginarySamples, float* RESTRICT OutPowerSamples, const int32 InNum);
+	SIGNALPROCESSING_API void BufferComplexToPowerFast(const float* RESTRICT InRealSamples, const float* RESTRICT InImaginarySamples, float* RESTRICT OutPowerSamples, const int32 InNum);
 
 	/** Compute magnitude of complex data. Out[i] = Sqrt(Real[i] * Real[i] + Imaginary[i] * Imaginary[i]) */
-	void SIGNALPROCESSING_API BufferComplexToMagnitudeFast(const AlignedFloatBuffer& InRealSamples, const AlignedFloatBuffer& InImaginarySamples, AlignedFloatBuffer& OutPowerSamples);
+	SIGNALPROCESSING_API void BufferComplexToMagnitudeFast(const AlignedFloatBuffer& InRealSamples, const AlignedFloatBuffer& InImaginarySamples, AlignedFloatBuffer& OutPowerSamples);
 
 	/** Compute magnitude of complex data. Out[i] = Sqrt(Real[i] * Real[i] + Imaginary[i] * Imaginary[i]) */
-	void SIGNALPROCESSING_API BufferComplexToMagnitudeFast(const float* RESTRICT InRealSamples, const float* RESTRICT InImaginarySamples, float* RESTRICT OutPowerSamples, const int32 InNum);
+	SIGNALPROCESSING_API void BufferComplexToMagnitudeFast(const float* RESTRICT InRealSamples, const float* RESTRICT InImaginarySamples, float* RESTRICT OutPowerSamples, const int32 InNum);
 }
