Refactor of the audio LinearEase and ExponentialEase classes in DSP.h to make the interfaces more indicitive of their behavior.

diff --git a/Engine/Plugins/Runtime/Synthesis/Source/Synthesis/Private/EpicSynth1.cpp b/Engine/Plugins/Runtime/Synthesis/Source/Synthesis/Private/EpicSynth1.cpp
--- a/Engine/Plugins/Runtime/Synthesis/Source/Synthesis/Private/EpicSynth1.cpp
+++ b/Engine/Plugins/Runtime/Synthesis/Source/Synthesis/Private/EpicSynth1.cpp
@@ -619,7 +619,7 @@ namespace Audio
 
 		for (int32 i = 0; i < 2; ++i)
 		{
-			Oscil[i].SetFrequency(PortamentoFrequency.GetValue());
+			Oscil[i].SetFrequency(PortamentoFrequency.GetNextValue());
 		}
 
 		FModulationMatrix* ModMatrix = &ParentSynth->ModMatrix;
diff --git a/Engine/Plugins/Runtime/Synthesis/Source/Synthesis/Private/SubmixEffectDelay.cpp b/Engine/Plugins/Runtime/Synthesis/Source/Synthesis/Private/SubmixEffectDelay.cpp
--- a/Engine/Plugins/Runtime/Synthesis/Source/Synthesis/Private/SubmixEffectDelay.cpp
+++ b/Engine/Plugins/Runtime/Synthesis/Source/Synthesis/Private/SubmixEffectDelay.cpp
@@ -48,7 +48,7 @@ void FSubmixEffectDelay::OnProcessAudio(const FSoundEffectSubmixInputData& InDat
 		{
 			const int32 SampleIndex = OutputBufferIndex + DelayIndex;
 			OutBuffer[SampleIndex] = DelaysPtr[DelayIndex].ProcessAudioSample(InBuffer[SampleIndex]);
-			DelayLines[DelayIndex].SetDelayMsec(InterpolationInfo.GetValue());
+			DelayLines[DelayIndex].SetDelayMsec(InterpolationInfo.GetNextValue());
 		}
 	}
 }
diff --git a/Engine/Plugins/Runtime/Synthesis/Source/Synthesis/Private/SubmixEffectTapDelay.cpp b/Engine/Plugins/Runtime/Synthesis/Source/Synthesis/Private/SubmixEffectTapDelay.cpp
--- a/Engine/Plugins/Runtime/Synthesis/Source/Synthesis/Private/SubmixEffectTapDelay.cpp
+++ b/Engine/Plugins/Runtime/Synthesis/Source/Synthesis/Private/SubmixEffectTapDelay.cpp
@@ -426,7 +426,7 @@ void FTapDelayInterpolationInfo::SetGainValue(float Value, float InterpolationTi
 
 float FTapDelayInterpolationInfo::GetGainValue()
 {
-	return GainParam.GetValue();
+	return GainParam.GetNextValue();
 }
 
 void FTapDelayInterpolationInfo::SetLengthValue(float Value, float InterpolationTime)
@@ -436,5 +436,5 @@ void FTapDelayInterpolationInfo::SetLengthValue(float Value, float Interpolation
 
 float FTapDelayInterpolationInfo::GetLengthValue()
 {
-	return LengthParam.GetValue();
+	return LengthParam.GetNextValue();
 }
\ No newline at end of file
diff --git a/Engine/Source/Runtime/SignalProcessing/Private/Chorus.cpp b/Engine/Source/Runtime/SignalProcessing/Private/Chorus.cpp
--- a/Engine/Source/Runtime/SignalProcessing/Private/Chorus.cpp
+++ b/Engine/Source/Runtime/SignalProcessing/Private/Chorus.cpp
@@ -89,7 +89,7 @@ namespace Audio
 			LFOQuadPhaseOut2 = Audio::GetUnipolar(-LFOQuadPhaseOut);
 			LFOQuadPhaseOut = Audio::GetUnipolar(LFOQuadPhaseOut);
 
-			float DepthValue = Depth[i].GetValue();
+			float DepthValue = Depth[i].GetNextValue();
 
 			if (i == EChorusDelays::Left)
 			{
diff --git a/Engine/Source/Runtime/SignalProcessing/Private/Delay.cpp b/Engine/Source/Runtime/SignalProcessing/Private/Delay.cpp
--- a/Engine/Source/Runtime/SignalProcessing/Private/Delay.cpp
+++ b/Engine/Source/Runtime/SignalProcessing/Private/Delay.cpp
@@ -180,7 +180,7 @@ namespace Audio
 			// If we're easing, then get the delay based on the current value of the ease
 			if (!EaseDelayMsec.IsDone())
 			{
-				DelayInSamples = EaseDelayMsec.GetValue() * SampleRate * 0.001f;
+				DelayInSamples = EaseDelayMsec.GetNextValue() * SampleRate * 0.001f;
 			}
 
 			DelayInSamples = FMath::Clamp(DelayInSamples, 0.0f, (float)(AudioBufferSize - 1));
diff --git a/Engine/Source/Runtime/SignalProcessing/Private/DynamicDelayAPF.cpp b/Engine/Source/Runtime/SignalProcessing/Private/DynamicDelayAPF.cpp
--- a/Engine/Source/Runtime/SignalProcessing/Private/DynamicDelayAPF.cpp
+++ b/Engine/Source/Runtime/SignalProcessing/Private/DynamicDelayAPF.cpp
@@ -105,8 +105,8 @@ void FDynamicDelayAPF::ProcessAudioBlock(const float* InSamples, const AlignedFl
 	DelayLineInput.AddUninitialized(InNum);
 
 	// Get G values to interpolate over
-	const float LastG = G.GetValue();
-	const float CurrG = G.GetValue(InNum);
+	const float LastG = G.GetNextValue();
+	const float CurrG = G.GetNextValue(InNum - 1);
 
 	// WorkBufferA = G * WorkBufferB
 	FMemory::Memcpy(WorkBufferA.GetData(), WorkBufferB.GetData(), InNum * sizeof(float));
diff --git a/Engine/Source/Runtime/SignalProcessing/Private/Granulator.cpp b/Engine/Source/Runtime/SignalProcessing/Private/Granulator.cpp
--- a/Engine/Source/Runtime/SignalProcessing/Private/Granulator.cpp
+++ b/Engine/Source/Runtime/SignalProcessing/Private/Granulator.cpp
@@ -843,7 +843,7 @@ namespace Audio
 				// If we're lerping to a new seek playhead time frame
 				if (!SeekingPlayheadTimeFrame.IsDone())
 				{
-					const float NewPlayheadFrame = SeekingPlayheadTimeFrame.GetValue();
+					const float NewPlayheadFrame = SeekingPlayheadTimeFrame.GetNextValue();
 
 					CurrentPlayHeadFrame = GetWrappedPlayheadPosition(NewPlayheadFrame);
 				}
diff --git a/Engine/Source/Runtime/SignalProcessing/Private/Osc.cpp b/Engine/Source/Runtime/SignalProcessing/Private/Osc.cpp
--- a/Engine/Source/Runtime/SignalProcessing/Private/Osc.cpp
+++ b/Engine/Source/Runtime/SignalProcessing/Private/Osc.cpp
@@ -278,7 +278,7 @@ namespace Audio
 				float SquareSaw1 = GetBipolar(Phase);
 				SquareSaw1 += PolySmooth(Phase, PhaseInc);
 
-				float CurrentPulseWidth = PulseWidthLerped.GetValue();
+				float CurrentPulseWidth = PulseWidthLerped.GetNextValue();
 
 				// Create a second sawtooth that is phase-shifted based on the pulsewidth
 				float NewPhase = 0.0f;
diff --git a/Engine/Source/Runtime/SignalProcessing/Private/SampleBufferReader.cpp b/Engine/Source/Runtime/SignalProcessing/Private/SampleBufferReader.cpp
--- a/Engine/Source/Runtime/SignalProcessing/Private/SampleBufferReader.cpp
+++ b/Engine/Source/Runtime/SignalProcessing/Private/SampleBufferReader.cpp
@@ -240,7 +240,7 @@ namespace Audio
 		int32 OutSampleIndex = 0;
 		for (int32 i = 0; i < NumFrames && !bIsFinished; ++i)
 		{
-			float CurrentPitch = Pitch.GetValue();
+			float CurrentPitch = Pitch.GetNextValue();
 
 			// Don't let the pitch go to 0.
 			if (FMath::IsNearlyZero(CurrentPitch))
diff --git a/Engine/Source/Runtime/SignalProcessing/Public/DSP/Dsp.h b/Engine/Source/Runtime/SignalProcessing/Public/DSP/Dsp.h
--- a/Engine/Source/Runtime/SignalProcessing/Public/DSP/Dsp.h
+++ b/Engine/Source/Runtime/SignalProcessing/Public/DSP/Dsp.h
@@ -242,6 +242,8 @@ namespace Audio
 			, Threshold(InThreshold)
 			, TargetValue(InInitValue)
 			, EaseFactor(InEaseFactor)
+			, OneMinusEase(1.0f - InEaseFactor)
+			, EaseTimesTarget(EaseFactor * InInitValue)
 		{
 		}
 
@@ -249,6 +251,10 @@ namespace Audio
 		{
 			CurrentValue = InInitValue;
 			TargetValue = InInitValue;
+			EaseFactor = InEaseFactor;
+
+			OneMinusEase = 1.0f - EaseFactor;
+			EaseTimesTarget = TargetValue * EaseFactor;
 		}
 
 		bool IsDone() const
@@ -256,19 +262,46 @@ namespace Audio
 			return FMath::Abs(TargetValue - CurrentValue) < Threshold;
 		}
 
-		float GetValue()
+		float GetNextValue()
 		{
 			if (IsDone())
 			{
 				return CurrentValue;
 			}
 
+			// Micro-optimization,
+			// But since GetNextValue(NumTicksToJumpAhead) does this work in a tight loop (non-vectorizable), might as well
+			/*
 			return CurrentValue = CurrentValue + (TargetValue - CurrentValue) * EaseFactor;
+								= CurrentValue + EaseFactor*TargetValue - EaseFactor*CurrentValue
+								= (CurrentValue - EaseFactor*CurrentValue) + EaseFactor*TargetValue
+								= (1 - EaseFactor)*CurrentValue + EaseFactor*TargetValue
+			*/
+			return CurrentValue = OneMinusEase * CurrentValue + EaseTimesTarget;
+		}
+
+		// same as GetValue(), but overloaded to jump forward by NumTicksToJumpAhead timesteps
+		// (before getting the value)
+		float GetNextValue(uint32 NumTicksToJumpAhead)
+		{
+			while (NumTicksToJumpAhead && !IsDone())
+			{
+				CurrentValue = OneMinusEase * CurrentValue + EaseTimesTarget;
+				--NumTicksToJumpAhead;
+			}
+
+			return CurrentValue;
+		}
+
+		float PeekCurrentValue() const
+		{
+			return CurrentValue;
 		}
 
 		void SetEaseFactor(const float InEaseFactor)
 		{
 			EaseFactor = InEaseFactor;
+			OneMinusEase = 1.0f - EaseFactor;
 		}
 
 		void operator=(const float& InValue)
@@ -279,6 +312,7 @@ namespace Audio
 		void SetValue(const float InValue, const bool bIsInit = false)
 		{
 			TargetValue = InValue;
+			EaseTimesTarget = EaseFactor * TargetValue;
 			if (bIsInit)
 			{
 				CurrentValue = TargetValue;
@@ -305,6 +339,12 @@ namespace Audio
 
 		// Percentage to move toward target value from current value each tick
 		float EaseFactor;
+
+		// 1.0f - EaseFactor
+		float OneMinusEase;
+
+		// EaseFactor * TargetValue
+		float EaseTimesTarget;
 	};
 	
 	// Simple easing function used to help interpolate params
@@ -345,7 +385,7 @@ namespace Audio
 			SetValue(End, InTimeSec);
 		}
 
-		float GetValue()
+		float GetNextValue()
 		{
 			if (IsDone())
 			{
@@ -359,21 +399,21 @@ namespace Audio
 		}
 
 		// same as GetValue(), but overloaded to increment Current Tick by NumTicksToJumpAhead
-		// (before getting the value);
-		float GetValue(int32 NumTicksToJumpAhead)
+		// (before getting the value)
+		float GetNextValue(int32 NumTicksToJumpAhead)
 		{
 			if (IsDone())
 			{
 				return CurrentValue;
 			}
 
-			CurrentTick += NumTicksToJumpAhead;
+			CurrentTick = FMath::Min(CurrentTick + NumTicksToJumpAhead, DurationTicks);
 			CurrentValue = DeltaValue * (float)CurrentTick / DurationTicks + StartValue;
 
 			return CurrentValue;
 		}
 
-		float GetCurrentValue() const
+		float PeekCurrentValue() const
 		{
 			return CurrentValue;
 		}
