Back out changelist 8976833

diff --git a/Engine/Plugins/Runtime/AudioSynesthesia/Source/AudioSynesthesia/Classes/LoudnessNRT.h b/Engine/Plugins/Runtime/AudioSynesthesia/Source/AudioSynesthesia/Classes/LoudnessNRT.h
deleted file mode 100644
diff --git a/Engine/Plugins/Runtime/AudioSynesthesia/Source/AudioSynesthesia/Private/AudioSynesthesiaModule.cpp b/Engine/Plugins/Runtime/AudioSynesthesia/Source/AudioSynesthesia/Private/AudioSynesthesiaModule.cpp
--- a/Engine/Plugins/Runtime/AudioSynesthesia/Source/AudioSynesthesia/Private/AudioSynesthesiaModule.cpp
+++ b/Engine/Plugins/Runtime/AudioSynesthesia/Source/AudioSynesthesia/Private/AudioSynesthesiaModule.cpp
@@ -4,7 +4,6 @@
 #include "CoreMinimal.h"
 #include "Modules/ModuleInterface.h"
 #include "Features/IModularFeatures.h"
-#include "LoudnessNRTFactory.h"
 
 DEFINE_LOG_CATEGORY(LogAudioSynesthesia);
 
@@ -14,13 +13,9 @@ namespace Audio
 {
 	void FAudioSynesthesiaModule::StartupModule()
 	{
-		// Register factories on startup
-		IModularFeatures::Get().RegisterModularFeature(FLoudnessNRTFactory::GetModularFeatureName(), &LoudnessFactory);
 	}
 
 	void FAudioSynesthesiaModule::ShutdownModule()
 	{
-		// Unregister factories on shutdown
-		IModularFeatures::Get().UnregisterModularFeature(FLoudnessNRTFactory::GetModularFeatureName(), &LoudnessFactory);
 	}
 }
diff --git a/Engine/Plugins/Runtime/AudioSynesthesia/Source/AudioSynesthesia/Private/LoudnessAnalyzer.cpp b/Engine/Plugins/Runtime/AudioSynesthesia/Source/AudioSynesthesia/Private/LoudnessAnalyzer.cpp
deleted file mode 100644
diff --git a/Engine/Plugins/Runtime/AudioSynesthesia/Source/AudioSynesthesia/Private/LoudnessNRT.cpp b/Engine/Plugins/Runtime/AudioSynesthesia/Source/AudioSynesthesia/Private/LoudnessNRT.cpp
deleted file mode 100644
diff --git a/Engine/Plugins/Runtime/AudioSynesthesia/Source/AudioSynesthesia/Private/LoudnessNRTFactory.cpp b/Engine/Plugins/Runtime/AudioSynesthesia/Source/AudioSynesthesia/Private/LoudnessNRTFactory.cpp
deleted file mode 100644
diff --git a/Engine/Plugins/Runtime/AudioSynesthesia/Source/AudioSynesthesia/Public/AudioSynesthesiaModule.h b/Engine/Plugins/Runtime/AudioSynesthesia/Source/AudioSynesthesia/Public/AudioSynesthesiaModule.h
--- a/Engine/Plugins/Runtime/AudioSynesthesia/Source/AudioSynesthesia/Public/AudioSynesthesiaModule.h
+++ b/Engine/Plugins/Runtime/AudioSynesthesia/Source/AudioSynesthesia/Public/AudioSynesthesiaModule.h
@@ -5,7 +5,6 @@
 #include "CoreMinimal.h"
 #include "Modules/ModuleInterface.h"
 #include "Modules/ModuleManager.h"
-#include "LoudnessNRTFactory.h"
 
 DECLARE_LOG_CATEGORY_EXTERN(LogAudioSynesthesia, Log, All);
 
@@ -20,6 +19,5 @@ namespace Audio
 		virtual void ShutdownModule() override;
 
 	private:
-		FLoudnessNRTFactory LoudnessFactory;
 	};
 }
diff --git a/Engine/Plugins/Runtime/AudioSynesthesia/Source/AudioSynesthesia/Public/InterpolateSorted.h b/Engine/Plugins/Runtime/AudioSynesthesia/Source/AudioSynesthesia/Public/InterpolateSorted.h
deleted file mode 100644
diff --git a/Engine/Plugins/Runtime/AudioSynesthesia/Source/AudioSynesthesia/Public/LoudnessAnalyzer.h b/Engine/Plugins/Runtime/AudioSynesthesia/Source/AudioSynesthesia/Public/LoudnessAnalyzer.h
deleted file mode 100644
diff --git a/Engine/Plugins/Runtime/AudioSynesthesia/Source/AudioSynesthesia/Public/LoudnessNRTFactory.h b/Engine/Plugins/Runtime/AudioSynesthesia/Source/AudioSynesthesia/Public/LoudnessNRTFactory.h
deleted file mode 100644
diff --git a/Engine/Source/Runtime/AudioAnalyzer/Private/AudioAnalyzerNRT.cpp b/Engine/Source/Runtime/AudioAnalyzer/Private/AudioAnalyzerNRT.cpp
--- a/Engine/Source/Runtime/AudioAnalyzer/Private/AudioAnalyzerNRT.cpp
+++ b/Engine/Source/Runtime/AudioAnalyzer/Private/AudioAnalyzerNRT.cpp
@@ -219,9 +219,10 @@ void UAudioAnalyzerNRT::Serialize(FArchive& Ar)
 			if (nullptr != Factory)
 			{
 				// Create result and worker from factory
+				TUniquePtr<Audio::IAnalyzerNRTResult> NewResult = Factory->NewResult();
 				{
 					FScopeLock ResultLock(&ResultCriticalSection);
-					Result = Factory->NewResultShared<ESPMode::ThreadSafe>();
+					Result = TSharedPtr<Audio::IAnalyzerNRTResult, ESPMode::ThreadSafe>(NewResult.Release());
 				}
 			}
 		}
diff --git a/Engine/Source/Runtime/AudioAnalyzer/Public/IAudioAnalyzerNRTInterface.h b/Engine/Source/Runtime/AudioAnalyzer/Public/IAudioAnalyzerNRTInterface.h
--- a/Engine/Source/Runtime/AudioAnalyzer/Public/IAudioAnalyzerNRTInterface.h
+++ b/Engine/Source/Runtime/AudioAnalyzer/Public/IAudioAnalyzerNRTInterface.h
@@ -5,6 +5,7 @@
 #pragma once
 
 #include "CoreMinimal.h"
+#include "Sound/SampleBuffer.h"
 #include "Features/IModularFeature.h"
 
 namespace Audio
@@ -103,15 +104,6 @@ namespace Audio
 		// Create a new result.
 		virtual TUniquePtr<IAnalyzerNRTResult> NewResult() = 0;
 
-		// Convenience function to create a new shared result by calling NewResult.
-		template<ESPMode Mode = ESPMode::Fast>
-		TSharedPtr<IAnalyzerNRTResult, Mode> NewResultShared()
-		{
-			TUniquePtr<Audio::IAnalyzerNRTResult> Result = NewResult();
-
-			return TSharedPtr<Audio::IAnalyzerNRTResult, Mode>(Result.Release());
-		}
-
 		// Create a new worker.
 		virtual TUniquePtr<IAnalyzerNRTWorker> NewWorker(const FAnalyzerNRTParameters& InParams, const IAnalyzerNRTSettings* InSettings) = 0;
 	};
diff --git a/Engine/Source/Runtime/SignalProcessing/Private/FloatArrayMath.cpp b/Engine/Source/Runtime/SignalProcessing/Private/FloatArrayMath.cpp
--- a/Engine/Source/Runtime/SignalProcessing/Private/FloatArrayMath.cpp
+++ b/Engine/Source/Runtime/SignalProcessing/Private/FloatArrayMath.cpp
@@ -10,10 +10,10 @@ namespace
 
 namespace Audio
 {
-	void ArrayCumulativeSum(TArrayView<const float> InView, TArray<float>& OutData)
+	void ArrayCumulativeSum(const TArray<float>& InData, TArray<float>& OutData)
 	{
 		// Initialize output data
-		int32 Num = InView.Num();
+		int32 Num = InData.Num();
 		OutData.Reset();
 		OutData.AddUninitialized(Num);
 
@@ -23,19 +23,19 @@ namespace Audio
 		}
 
 		float* OutDataPtr = OutData.GetData();
-		const float* InViewPtr = InView.GetData();
+		const float* InDataPtr = InData.GetData();
 
 		// Start summing
-		*OutDataPtr = *InViewPtr++;
+		*OutDataPtr = *InDataPtr++;
 
 		for (int32 i = 1; i < Num; i++)
 		{
-			float Temp = *OutDataPtr++ + *InViewPtr++;
+			float Temp = *OutDataPtr++ + *InDataPtr++;
 			*OutDataPtr = Temp;
 		}
 	}
 
-	void ArrayMeanFilter(TArrayView<const float> InView, int32 WindowSize, int32 WindowOrigin, TArray<float>& OutData)
+	void ArrayMeanFilter(const TArray<float>& InData, int32 WindowSize, int32 WindowOrigin, TArray<float>& OutData)
 	{
 		// a quick but sinful implementation of a mean filter. encourages floating point rounding errors. 
 		check(WindowOrigin < WindowSize);
@@ -43,7 +43,7 @@ namespace Audio
 		check(WindowSize > 0);
 
 		// Initialize output data
-		const int32 Num = InView.Num();
+		const int32 Num = InData.Num();
 		OutData.Reset();
 		OutData.AddUninitialized(Num);
 
@@ -53,11 +53,11 @@ namespace Audio
 		}
 		
 		// Use cumulative sum to avoid multiple summations 
-		// Instead of summing over InView[StartIndex:EndIndex], avoid all that
+		// Instead of summing over InData[StartIndex:EndIndex], avoid all that
 		// calculation by taking difference of cumulative sum at those two points:
 		//  cumsum(X[0:b]) - cumsum(X[0:a]) = sum(X[a:b])
 		TArray<float> SummedData;
-		ArrayCumulativeSum(InView, SummedData);
+		ArrayCumulativeSum(InData, SummedData);
 		const float LastSummedData = SummedData.Last();
 		
 		
@@ -101,7 +101,7 @@ namespace Audio
 		}
 	}
 
-	void ArrayMaxFilter(TArrayView<const float> InView, int32 WindowSize, int32 WindowOrigin, TArray<float>& OutData)
+	void ArrayMaxFilter(const TArray<float>& InData, int32 WindowSize, int32 WindowOrigin, TArray<float>& OutData)
 	{
 		// A reasonable implementation of a max filter for the data we're interested in, though surely not the fastest.
 		check(WindowOrigin < WindowSize);
@@ -112,7 +112,7 @@ namespace Audio
 		int32 EndIndex = StartIndex + WindowSize;
 
 		// Initialize output
-		int32 Num = InView.Num();
+		int32 Num = InData.Num();
 		OutData.Reset();
 		OutData.AddUninitialized(Num);
 
@@ -125,16 +125,16 @@ namespace Audio
 		int32 ActualStartIndex = 0;
 		int32 ActualEndIndex = FMath::Min(EndIndex, Num);
 
-		const float* InViewPtr = InView.GetData();
+		const float* InDataPtr = InData.GetData();
 		float* OutDataPtr = OutData.GetData();
 		int32 MaxIndex = 0;
-		float MaxValue = InView[0];
+		float MaxValue = InData[0];
 
 		for (int32 i = ActualStartIndex; i < ActualEndIndex; i++)
 		{
-			if (InViewPtr[i] > MaxValue)
+			if (InDataPtr[i] > MaxValue)
 			{
-				MaxValue = InViewPtr[i];
+				MaxValue = InDataPtr[i];
 				MaxIndex = i;
 			}		
 		}
@@ -153,23 +153,23 @@ namespace Audio
 			{
 				// We need to evaluate the entire window because the previous maximum value was not in this window.
 				MaxIndex = ActualStartIndex;
-				MaxValue = InViewPtr[MaxIndex];
+				MaxValue = InDataPtr[MaxIndex];
 				for (int32 j = ActualStartIndex + 1; j < ActualEndIndex; j++)
 				{
-					if (InViewPtr[j] > MaxValue)
+					if (InDataPtr[j] > MaxValue)
 					{
 						MaxIndex = j;
-						MaxValue = InViewPtr[MaxIndex];
+						MaxValue = InDataPtr[MaxIndex];
 					}
 				}
 			}
 			else
 			{
 				// We only need to inspect the newest sample because the previous maximum value was in this window.
-				if (InViewPtr[ActualEndIndex - 1] > MaxValue)
+				if (InDataPtr[ActualEndIndex - 1] > MaxValue)
 				{
 					MaxIndex = ActualEndIndex - 1;
-					MaxValue = InViewPtr[MaxIndex];
+					MaxValue = InDataPtr[MaxIndex];
 				}
 			}
 
@@ -180,70 +180,59 @@ namespace Audio
 		}
 	}
 
-	void ArrayGetEuclideanNorm(TArrayView<const float> InView, float& OutEuclideanNorm)
+	void ArrayGetEuclideanNorm(const TArray<float>& InArray, float& OutEuclideanNorm)
 	{
 		// Initialize output.
 		OutEuclideanNorm = 0.0f;
-		const int32 Num = InView.Num();
-		const float* InViewData = InView.GetData();
+		const int32 Num = InArray.Num();
+		const float* InArrayData = InArray.GetData();
 		
 		// Sum it up.
 		for (int32 i = 0; i < Num; i++)
 		{
-			OutEuclideanNorm += InViewData[i] * InViewData[i];
+			OutEuclideanNorm += InArrayData[i] * InArrayData[i];
 		}
 
 		OutEuclideanNorm = FMath::Sqrt(OutEuclideanNorm);
 	}
 
-	void ArrayClampInPlace(TArrayView<float> InView, float InMin, float InMax)
+	void ArrayMultiplyByConstantInPlace(TArray<float>& InArray, float InMultiplier)
 	{
-		const int32 Num = InView.Num();
-		float* Data = InView.GetData();
-
-		for (int32 i = 0; i < Num; i++)
-		{
-			Data[i] = FMath::Clamp(Data[i], InMin, InMax);
-		}
-	}
-
-	void ArrayMultiplyByConstantInPlace(TArrayView<float> InView, float InMultiplier)
-	{
-		const int32 Num = InView.Num();
-		float* InViewData = InView.GetData();
+		const int32 Num = InArray.Num();
+		float* InArrayData = InArray.GetData();
 		for (int32 i = 0; i < Num; i++)
 		{
-			InViewData[i] *= InMultiplier;
+			InArrayData[i] *= InMultiplier;
 		}
 	}
 
-	void ArraySubtractByConstantInPlace(TArrayView<float> InView, float InSubtrahend)
+	void ArraySubtractByConstantInPlace(TArray<float>& InArray, float InSubtrahend)
 	{
-		const int32 Num = InView.Num();
-		float* InViewData = InView.GetData();
+		const int32 Num = InArray.Num();
+		float* InArrayData = InArray.GetData();
 		for (int32 i = 0; i < Num; i++)
 		{
-			InViewData[i] -= InSubtrahend;
+			InArrayData[i] -= InSubtrahend;
 		}
 	}
 
-	void ArrayMagnitudeToDecibelInPlace(TArrayView<float> InView)
+	void ArrayMagnitudeToDecibelInPlace(TArray<float>& InArray)
 	{
-		const int32 Num = InView.Num();
-		float* InViewData = InView.GetData();
+		const int32 Num = InArray.Num();
+		float* InArrayData = InArray.GetData();
 		for (int32 i = 0; i < Num; i++)
 		{
-			InViewData[i] = 20.f * FMath::Loge(InViewData[i]) / LOGE10;
+			InArrayData[i] = 20.f * FMath::Loge(InArrayData[i]) / LOGE10;
 		}
 	}
 
-	void ArrayPowerToDecibelInPlace(TArrayView<float> InView)
+	void ArrayPowerToDecibelInPlace(TArray<float>& InArray)
 	{
-		const int32 Num = InView.Num();
-		float* InViewData = InView.GetData();
+		const int32 Num = InArray.Num();
+		float* InArrayData = InArray.GetData();
 		for (int32 i = 0; i < Num; i++)
 		{
-			InViewData[i] = 10.f * FMath::Loge(InViewData[i]) / LOGE10;
+			InArrayData[i] = 10.f * FMath::Loge(InArrayData[i]) / LOGE10;
 		}
 	}
 
@@ -280,9 +269,9 @@ namespace Audio
 		Kernel[RowIndex].OffsetValues = TArray<float>(OffsetValues.GetData(), OffsetValues.Num());
 	}
 
-	void FContiguousSparse2DKernelTransform::TransformArray(TArrayView<const float> InView, TArray<float>& OutArray) const
+	void FContiguousSparse2DKernelTransform::TransformArray(TArrayView<const float> InArray, TArray<float>& OutArray) const
 	{
-		check(InView.Num() == NumIn);
+		check(InArray.Num() == NumIn);
 
 		// Resize output
 		OutArray.Reset(NumOut);
@@ -291,12 +280,12 @@ namespace Audio
 			OutArray.AddUninitialized(NumOut);
 		}
 
-		TransformArray(InView.GetData(), OutArray.GetData());
+		TransformArray(InArray.GetData(), OutArray.GetData());
 	}
 
-	void FContiguousSparse2DKernelTransform::TransformArray(TArrayView<const float> InView, AlignedFloatBuffer& OutArray) const
+	void FContiguousSparse2DKernelTransform::TransformArray(TArrayView<const float> InArray, AlignedFloatBuffer& OutArray) const
 	{	
-		check(InView.Num() == NumIn);
+		check(InArray.Num() == NumIn);
 
 		// Resize output
 		OutArray.Reset(NumOut);
@@ -305,7 +294,7 @@ namespace Audio
 			OutArray.AddUninitialized(NumOut);
 		}
 
-		TransformArray(InView.GetData(), OutArray.GetData());
+		TransformArray(InArray.GetData(), OutArray.GetData());
 	}
 
 	void FContiguousSparse2DKernelTransform::TransformArray(const float* InArray, float* OutArray) const
diff --git a/Engine/Source/Runtime/SignalProcessing/Public/DSP/FloatArrayMath.h b/Engine/Source/Runtime/SignalProcessing/Public/DSP/FloatArrayMath.h
--- a/Engine/Source/Runtime/SignalProcessing/Public/DSP/FloatArrayMath.h
+++ b/Engine/Source/Runtime/SignalProcessing/Public/DSP/FloatArrayMath.h
@@ -9,61 +9,53 @@ namespace Audio
 {
 	/** Cumulative sum of array.
 	 *
-	 *  InView contains data to be cumulatively summed.
-	 *  OutData contains sum and is same size as InView.
+	 *  InData contains data to be cumulatively summed.
+	 *  OutData contains sum and is same size as InData.
 	 */
-	SIGNALPROCESSING_API void ArrayCumulativeSum(TArrayView<const float> InView, TArray<float>& OutData);
+	SIGNALPROCESSING_API void ArrayCumulativeSum(const TArray<float>& InData, TArray<float>& OutData);
 
 	/** Mean filter of array.
 	 *
 	 *  Note: Uses standard biased mean estimator of Sum(x) / Count(x).
 	 *  Note: At array boundaries, this algorithm truncates windows where no valid array data exists. Values calculated with truncated windows have corresponding increased variances. 
 	 *
-	 *  InView contains data to be filtered.
-	 *  WindowSize determines the number of samples from InView analyzed to produce a value in OutData.
-	 *  WindowOrigin describes the offset from the windows first sample to the index of OutData. For example, if WindowOrigin = WindowSize/4, then OutData[i] = Mean(InView[i - Window/4 : i + 3 * Window / 4]).
+	 *  InData contains data to be filtered.
+	 *  WindowSize determines the number of samples from InData analyzed to produce a value in OutData.
+	 *  WindowOrigin describes the offset from the windows first sample to the index of OutData. For example, if WindowOrigin = WindowSize/4, then OutData[i] = Mean(InData[i - Window/4 : i + 3 * Window / 4]).
 	 *  OutData contains the produceds data.
 	 */
-	SIGNALPROCESSING_API void ArrayMeanFilter(TArrayView<const float> InView, int32 WindowSize, int32 WindowOrigin, TArray<float>& OutData);
+	SIGNALPROCESSING_API void ArrayMeanFilter(const TArray<float>& InData, int32 WindowSize, int32 WindowOrigin, TArray<float>& OutData);
 
 	/** Max filter of array.
 	 *
 	 *  Note: At array boundaries, this algorithm truncates windows where no valid array data exists. 
 	 *
-	 *  InView contains data to be filtered.
-	 *  WindowSize determines the number of samples from InView analyzed to produce a value in OutData.
-	 *  WindowOrigin describes the offset from the windows first sample to the index of OutData. For example, if WindowOrigin = WindowSize/4, then OutData[i] = Max(InView[i - Window/4 : i + 3 * Window / 4]).
+	 *  InData contains data to be filtered.
+	 *  WindowSize determines the number of samples from InData analyzed to produce a value in OutData.
+	 *  WindowOrigin describes the offset from the windows first sample to the index of OutData. For example, if WindowOrigin = WindowSize/4, then OutData[i] = Max(InData[i - Window/4 : i + 3 * Window / 4]).
 	 *  OutData contains the produceds data.
 	 */
-	SIGNALPROCESSING_API void ArrayMaxFilter(TArrayView<const float> InView, int32 WindowSize, int32 WindowOrigin, TArray<float>& OutData);
+	SIGNALPROCESSING_API void ArrayMaxFilter(const TArray<float>& InData, int32 WindowSize, int32 WindowOrigin, TArray<float>& OutData);
 
-	/** Computes the EuclideanNorm of the InView. Same as calculating the energy in window. */
-	SIGNALPROCESSING_API void ArrayGetEuclideanNorm(TArrayView<const float> InView, float& OutEuclideanNorm);
+	/** Computes the EuclideanNorm of the InArray. Same as calculating the energy in window. */
+	SIGNALPROCESSING_API void ArrayGetEuclideanNorm(const TArray<float>& InArray, float& OutEuclideanNorm);
 
-	/** Clamp values in an array.
-	 *
-	 *  InView is a view of a float array to be clamped.
-	 *  InMin is the minimum value allowable in the array.
-	 *  InMax is the maximum value allowable in the array.
-	 */
-	SIGNALPROCESSING_API void ArrayClampInPlace(TArrayView<float> InView, float InMin, float InMax);
-
-	/** Multiplies each element in InView by InMultiplier */
-	SIGNALPROCESSING_API void ArrayMultiplyByConstantInPlace(TArrayView<float> InView, float InMultiplier);
+	/** Multiplies each element in InArray by InMultiplier */
+	SIGNALPROCESSING_API void ArrayMultiplyByConstantInPlace(TArray<float>& InArray, float InMultiplier);
 
-	/** Subtract value from each element in InView */
-	SIGNALPROCESSING_API void ArraySubtractByConstantInPlace(TArrayView<float> InView, float InSubtrahend);
+	/** Subtract value from each element in InArray */
+	SIGNALPROCESSING_API void ArraySubtractByConstantInPlace(TArray<float>& InArray, float InSubtrahend);
 
 	/** Convert magnitude values to decibel values in place. db = 20 * log10(val) */
-	SIGNALPROCESSING_API void ArrayMagnitudeToDecibelInPlace(TArrayView<float> InView);
+	SIGNALPROCESSING_API void ArrayMagnitudeToDecibelInPlace(TArray<float>& InArray);
 
 	/** Convert power values to decibel values in place. db = 10 * log10(val) */
-	SIGNALPROCESSING_API void ArrayPowerToDecibelInPlace(TArrayView<float> InView);
+	SIGNALPROCESSING_API void ArrayPowerToDecibelInPlace(TArray<float>& InArray);
 
 	/** FContiguousSparse2DKernelTransform
 	 *
 	 *  FContiguousSparse2DKernelTransform applies a matrix transformation to an input array. 
-	 *  [OutArray] = [[Kernal]][InView]  
+	 *  [OutArray] = [[Kernal]][InArray]  
 	 *
 	 *  It provides some optimization by exploit the contiguous and sparse qualities of the kernel rows,
 	 *  which allows it to skip multiplications with the number zero. 
@@ -103,17 +95,17 @@ namespace Audio
 
 		/** Transforms the input array given the kernel.
 		 *
-		 *  InView is the array to be transformed. It must have `NumInElements` number of elements.
+		 *  InArray is the array to be transformed. It must have `NumInElements` number of elements.
 		 *  OutArray is the transformed array. It will have `NumOutElements` number of elements.
 		 */
-		void TransformArray(TArrayView<const float> InView, TArray<float>& OutArray) const;
+		void TransformArray(TArrayView<const float> InArray, TArray<float>& OutArray) const;
 
 		/** Transforms the input array given the kernel.
 		 *
-		 *  InView is the array to be transformed. It must have `NumInElements` number of elements.
+		 *  InArray is the array to be transformed. It must have `NumInElements` number of elements.
 		 *  OutArray is the transformed array. It will have `NumOutElements` number of elements.
 		 */
-		void TransformArray(TArrayView<const float> InView, AlignedFloatBuffer& OutArray) const;
+		void TransformArray(TArrayView<const float> InArray, AlignedFloatBuffer& OutArray) const;
 
 		/** Transforms the input array given the kernel.
 		 *
