small fixup: add const qualifiers to new getters in SampleBuffer

diff --git a/Engine/Source/Runtime/SignalProcessing/Public/SampleBuffer.h b/Engine/Source/Runtime/SignalProcessing/Public/SampleBuffer.h
--- a/Engine/Source/Runtime/SignalProcessing/Public/SampleBuffer.h
+++ b/Engine/Source/Runtime/SignalProcessing/Public/SampleBuffer.h
@@ -399,7 +399,7 @@ namespace Audio
 		// InIndex [0.0f, NumSamples - 1.0f]
 		// OutFrame is the multichannel output for one index value
 		// Returns InIndex wrapped between 0.0 and NumFrames
-		float GetAudioFrameAtFractionalIndex(float InIndex, TArray<SampleType>& OutFrame)
+		float GetAudioFrameAtFractionalIndex(float InIndex, TArray<SampleType>& OutFrame) const
 		{
 			InIndex = FMath::Fmod(InIndex, static_cast<float>(NumFrames));
 
@@ -411,7 +411,7 @@ namespace Audio
 		// InPhase [0, 1], wrapped, through duration of file (ignores sample rate)
 		// OutFrame is the multichannel output for one phase value
 		// Returns InPhase wrapped between 0.0 and 1.0
-		float GetAudioFrameAtPhase(float InPhase, TArray<SampleType>& OutFrame)
+		float GetAudioFrameAtPhase(float InPhase, TArray<SampleType>& OutFrame) const
 		{
 			InPhase = FMath::Fmod(InPhase, 1.0f);
 
@@ -424,7 +424,7 @@ namespace Audio
 		// InTimeSec, get the value of the buffer at the given time (uses sample rate)
 		// OutFrame is the multichannel output for one time value
 		// Returns InTimeSec wrapped between 0.0 and (NumSamples / SampleRate)
-		float GetAudioFrameAtTime(float InTimeSec, TArray<SampleType>& OutFrame)
+		float GetAudioFrameAtTime(float InTimeSec, TArray<SampleType>& OutFrame) const
 		{
 			if (InTimeSec >= SampleDuration)
 			{
@@ -441,7 +441,7 @@ namespace Audio
 	private:
 		// Internal implementation. Called by all public GetAudioFrameAt_ _ _ _() functions
 		// public functions do range checking/wrapping and then call this function
-		void GetAudioFrameAtFractionalIndexInternal(float InIndex, TArray<SampleType>& OutFrame)
+		void GetAudioFrameAtFractionalIndexInternal(float InIndex, TArray<SampleType>& OutFrame) const 
 		{
 			const float Alpha = FMath::Fmod(InIndex, 1.0f);
 			const int32 WholeThisIndex = FMath::FloorToInt(InIndex);
