[NO NON-PII COMMIT MESSAGE]

diff --git a/Engine/Plugins/Online/OnlineSubsystem/Source/OnlineSubsystem.Build.cs b/Engine/Plugins/Online/OnlineSubsystem/Source/OnlineSubsystem.Build.cs
--- a/Engine/Plugins/Online/OnlineSubsystem/Source/OnlineSubsystem.Build.cs
+++ b/Engine/Plugins/Online/OnlineSubsystem/Source/OnlineSubsystem.Build.cs
@@ -25,7 +25,8 @@ public class OnlineSubsystem : ModuleRules
 				"CoreUObject",
 				"ImageCore",
 				"Sockets",
-				"JsonUtilities"
+				"JsonUtilities",
+				"AudioMixer"
 			}
 		);
 	}
diff --git a/Engine/Plugins/Online/OnlineSubsystem/Source/Public/Interfaces/VoiceInterface.h b/Engine/Plugins/Online/OnlineSubsystem/Source/Public/Interfaces/VoiceInterface.h
--- a/Engine/Plugins/Online/OnlineSubsystem/Source/Public/Interfaces/VoiceInterface.h
+++ b/Engine/Plugins/Online/OnlineSubsystem/Source/Public/Interfaces/VoiceInterface.h
@@ -6,6 +6,7 @@
 #include "UObject/CoreOnline.h"
 #include "OnlineDelegateMacros.h"
 #include "OnlineSubsystemPackage.h"
+#include "DSP/MultithreadedPatching.h"
 
 class FVoicePacket;
 
@@ -230,6 +231,28 @@ public:
 	 * Get information about the voice state for display
 	 */
 	virtual FString GetVoiceDebugState() const = 0;
+
+	/**
+	 * This can be used to get arbitrarily tap outgoing microphone audio from local players to use elsewhere in your application. 
+	 * for information on usage, see MultithreadedPatching.h.
+	 *
+	 * @return nullptr if the VoiceEngine implementation being used does not support tapping audio output from the microphone.
+	 */
+	virtual Audio::FPatchOutputStrongPtr GetMicrophoneOutput()
+	{
+		return nullptr;
+	};
+
+
+	/**
+	 * This can be used to get arbitrarily tap incoming VOIP audio from remote players to use elsewhere in your application.
+	 * for information on usage, see MultithreadedPatching.h.
+	 * @return nullptr if the VoiceEngine implementation being used does not support tapping incoming audio.
+	 */
+	virtual Audio::FPatchOutputStrongPtr GetRemoteTalkerOutput()
+	{
+		return nullptr;
+	};
 };
 
 typedef TSharedPtr<IVoiceEngine, ESPMode::ThreadSafe> IVoiceEnginePtr;
@@ -437,6 +460,28 @@ public:
 	 * Get information about the voice state for display
 	 */
 	virtual FString GetVoiceDebugState() const = 0;
+
+	/**
+	 * This can be used to get arbitrarily tap outgoing microphone audio from local players to use elsewhere in your application.
+	 * for information on usage, see MultithreadedPatching.h.
+	 *
+	 * @return nullptr if the VoiceEngine implementation being used does not support tapping audio output from the microphone.
+	 */
+	virtual Audio::FPatchOutputStrongPtr GetMicrophoneOutput()
+	{
+		return nullptr;
+	};
+
+
+	/**
+	 * This can be used to get arbitrarily tap incoming VOIP audio from remote players to use elsewhere in your application.
+	 * for information on usage, see MultithreadedPatching.h.
+	 * @return nullptr if the VoiceEngine implementation being used does not support tapping incoming audio.
+	 */
+	virtual Audio::FPatchOutputStrongPtr GetRemoteTalkerOutput()
+	{
+		return nullptr;
+	};
 };
 
 typedef TSharedPtr<IOnlineVoice, ESPMode::ThreadSafe> IOnlineVoicePtr;
diff --git a/Engine/Plugins/Online/OnlineSubsystemUtils/Source/OnlineSubsystemUtils/Private/VoiceEngineImpl.cpp b/Engine/Plugins/Online/OnlineSubsystemUtils/Source/OnlineSubsystemUtils/Private/VoiceEngineImpl.cpp
--- a/Engine/Plugins/Online/OnlineSubsystemUtils/Source/OnlineSubsystemUtils/Private/VoiceEngineImpl.cpp
+++ b/Engine/Plugins/Online/OnlineSubsystemUtils/Source/OnlineSubsystemUtils/Private/VoiceEngineImpl.cpp
@@ -21,6 +21,7 @@ namespace NotificationClient
 }
 #endif 
 
+
 FRemoteTalkerDataImpl::FRemoteTalkerDataImpl() :
 	MaxUncompressedDataSize(0),
 	MaxUncompressedDataQueueSize(0),
@@ -557,6 +558,7 @@ uint32 FVoiceEngineImpl::SubmitRemoteVoiceData(const FUniqueNetIdWrapper& Remote
 			QueuedData.bIsActive = false;
 			QueuedData.VoipSynthComponent->OpenPacketStream(InSampleCount, UVOIPStatics::GetNumBufferedPackets(), UVOIPStatics::GetBufferingDelay());
 			QueuedData.bIsEnvelopeBound = false;
+			QueuedData.VoipSynthComponent->ConnectToSplitter(AllRemoteTalkerAudio);
 		}
 	}
 
@@ -793,6 +795,24 @@ bool FVoiceEngineImpl::Exec(UWorld* InWorld, const TCHAR* Cmd, FOutputDevice& Ar
 	return bWasHandled;
 }
 
+Audio::FPatchOutputStrongPtr FVoiceEngineImpl::GetMicrophoneOutput()
+{
+	 // NOTE: We don't mix down multiple microphones here.
+	if (VoiceCapture.IsValid())
+	{
+		return VoiceCapture->GetMicrophoneAudio(4096 * 2, 1.0f);
+	}
+	else
+	{
+		return nullptr;
+	}
+}
+
+Audio::FPatchOutputStrongPtr FVoiceEngineImpl::GetRemoteTalkerOutput()
+{
+	return AllRemoteTalkerAudio.AddNewOutput(4096 * 2, 1.0f);
+}
+
 int32 FVoiceEngineImpl::GetMaxVoiceRemainderSize()
 {
 	return MAX_VOICE_REMAINDER_SIZE;
diff --git a/Engine/Plugins/Online/OnlineSubsystemUtils/Source/OnlineSubsystemUtils/Private/VoiceInterfaceImpl.cpp b/Engine/Plugins/Online/OnlineSubsystemUtils/Source/OnlineSubsystemUtils/Private/VoiceInterfaceImpl.cpp
--- a/Engine/Plugins/Online/OnlineSubsystemUtils/Source/OnlineSubsystemUtils/Private/VoiceInterfaceImpl.cpp
+++ b/Engine/Plugins/Online/OnlineSubsystemUtils/Source/OnlineSubsystemUtils/Private/VoiceInterfaceImpl.cpp
@@ -882,3 +882,27 @@ FString FOnlineVoiceImpl::GetVoiceDebugState() const
 
 	return Output;
 }
+
+Audio::FPatchOutputStrongPtr FOnlineVoiceImpl::GetMicrophoneOutput()
+{
+	if (VoiceEngine.IsValid())
+	{
+		return VoiceEngine->GetMicrophoneOutput();
+	}
+	else
+	{
+		return nullptr;
+	}
+}
+
+Audio::FPatchOutputStrongPtr FOnlineVoiceImpl::GetRemoteTalkerOutput()
+{
+	if (VoiceEngine.IsValid())
+	{
+		return VoiceEngine->GetRemoteTalkerOutput();
+	}
+	else
+	{
+		return nullptr;
+	}
+}
diff --git a/Engine/Plugins/Online/OnlineSubsystemUtils/Source/OnlineSubsystemUtils/Private/VoipListenerSynthComponent.cpp b/Engine/Plugins/Online/OnlineSubsystemUtils/Source/OnlineSubsystemUtils/Private/VoipListenerSynthComponent.cpp
--- a/Engine/Plugins/Online/OnlineSubsystemUtils/Source/OnlineSubsystemUtils/Private/VoipListenerSynthComponent.cpp
+++ b/Engine/Plugins/Online/OnlineSubsystemUtils/Source/OnlineSubsystemUtils/Private/VoipListenerSynthComponent.cpp
@@ -10,6 +10,21 @@ static int16 DebugBuffer[2048];
 
 #define DEBUG_NOISE 0
 
+static int32 DefaultPatchBufferSizeCVar = 4096;
+FAutoConsoleVariableRef CVarDefaultPatchBufferSize(
+	TEXT("voice.DefaultPatchBufferSize"),
+	DefaultPatchBufferSizeCVar,
+	TEXT("Changes the amount of audio we buffer for VOIP patching, in samples.\n"),
+	ECVF_Default);
+
+static int32 DefaultPatchGainCVar = 1.0f;
+FAutoConsoleVariableRef CVarDefaultPatchGain(
+	TEXT("voice.DefaultPatchGain"),
+	DefaultPatchGainCVar,
+	TEXT("Changes the default gain of audio patches, in samples.\n"),
+	ECVF_Default);
+
+
 bool UVoipListenerSynthComponent::Init(int32& SampleRate)
 {
 	NumChannels = 1;
@@ -41,6 +56,14 @@ int32 UVoipListenerSynthComponent::OnGenerateAudio(float* OutAudio, int32 NumSam
 		OutAudio[Index] = 0.5 * OutAudio[Index] + 0.5 * FMath::FRand();
 	}
 #endif
+
+	// Push audio to the patch splitter.
+	if (ExternalSend.IsOutputStillActive())
+	{
+		ExternalSend.PushAudio(OutAudio, NumSamples);
+	}
+	
+
 	return NumSamples;
 }
 
@@ -91,6 +114,11 @@ void UVoipListenerSynthComponent::SubmitPacket(void* InBuffer, int32 NumBytes, i
 	}
 }
 
+void UVoipListenerSynthComponent::ConnectToSplitter(Audio::FPatchMixerSplitter& InSplitter)
+{
+	ExternalSend = InSplitter.AddNewInput(DefaultPatchBufferSizeCVar, DefaultPatchGainCVar);
+}
+
 bool UVoipListenerSynthComponent::IsIdling()
 {
 	if (PacketBuffer.IsValid())
diff --git a/Engine/Plugins/Online/OnlineSubsystemUtils/Source/OnlineSubsystemUtils/Public/VoiceEngineImpl.h b/Engine/Plugins/Online/OnlineSubsystemUtils/Source/OnlineSubsystemUtils/Public/VoiceEngineImpl.h
--- a/Engine/Plugins/Online/OnlineSubsystemUtils/Source/OnlineSubsystemUtils/Public/VoiceEngineImpl.h
+++ b/Engine/Plugins/Online/OnlineSubsystemUtils/Source/OnlineSubsystemUtils/Public/VoiceEngineImpl.h
@@ -14,6 +14,7 @@
 #include "VoipListenerSynthComponent.h"
 #include "OnlineSubsystemUtilsPackage.h"
 #include "AudioDevice.h"
+#include "DSP/MultithreadedPatching.h"
 
 #include "VoicePacketImpl.h"
 #include "UObject/CoreOnline.h"
@@ -37,6 +38,8 @@ struct FLocalVoiceData
 	uint32 VoiceRemainderSize;
 	/** Voice sample data not encoded last time */
 	TArray<uint8> VoiceRemainder;
+	/** Output for a local talker. */
+	Audio::FPatchSplitter LocalVoiceOutput;
 };
 
 /** 
@@ -81,6 +84,8 @@ public:
 	TArray<uint8> UncompressedDataQueue;
 	/** Per remote talker voice decoding state */
 	TSharedPtr<IVoiceDecoder> VoiceDecoder;
+	/** Patch splitter to expose incoming audio to multiple outputs. */
+	Audio::FPatchSplitter RemoteVoiceOutput;
 };
 
 /**
@@ -123,6 +128,9 @@ class ONLINESUBSYSTEMUTILS_API FVoiceEngineImpl : public IVoiceEngine, public FS
 	/** Mapping of UniqueIds to the incoming voice data and their audio component */
 	typedef TMap<FUniqueNetIdWrapper, FRemoteTalkerDataImpl> FRemoteTalkerData;
 
+	/** Total mixed down output of all remote player talkers. */
+	Audio::FPatchMixerSplitter MixedDownRemoteTalkerOutput;
+
 	/** Reference to the main online subsystem */
 	IOnlineSubsystem* OnlineSubsystem;
 
@@ -154,6 +162,9 @@ class ONLINESUBSYSTEMUTILS_API FVoiceEngineImpl : public IVoiceEngine, public FS
 	/** Serialization helper */
 	FVoiceSerializeHelper* SerializeHelper;
 
+	/** Audio taps for the full mixdown of all remote players. */
+	Audio::FPatchMixerSplitter AllRemoteTalkerAudio;
+
 // Get Audio Device Changes on Windows
 #if PLATFORM_WINDOWS
 	FThreadSafeBool bAudioDeviceChanged;
@@ -282,6 +293,10 @@ public:
 
 	virtual void GetVoiceSettingsOverride(const FUniqueNetIdWrapper& RemoteTalkerId, FVoiceSettings& VoiceSettings) {}
 
+
+	virtual Audio::FPatchOutputStrongPtr GetMicrophoneOutput() override;
+	virtual Audio::FPatchOutputStrongPtr GetRemoteTalkerOutput() override;
+
 private:
 
 	/**
diff --git a/Engine/Plugins/Online/OnlineSubsystemUtils/Source/OnlineSubsystemUtils/Public/VoiceInterfaceImpl.h b/Engine/Plugins/Online/OnlineSubsystemUtils/Source/OnlineSubsystemUtils/Public/VoiceInterfaceImpl.h
--- a/Engine/Plugins/Online/OnlineSubsystemUtils/Source/OnlineSubsystemUtils/Public/VoiceInterfaceImpl.h
+++ b/Engine/Plugins/Online/OnlineSubsystemUtils/Source/OnlineSubsystemUtils/Public/VoiceInterfaceImpl.h
@@ -151,6 +151,8 @@ public:
 	virtual void ClearVoicePackets() override;
 	virtual void Tick(float DeltaTime) override;
 	virtual FString GetVoiceDebugState() const override;
+	virtual Audio::FPatchOutputStrongPtr GetMicrophoneOutput() override;
+	virtual Audio::FPatchOutputStrongPtr GetRemoteTalkerOutput() override;
 };
 
 typedef TSharedPtr<FOnlineVoiceImpl, ESPMode::ThreadSafe> FOnlineVoiceImplPtr;
diff --git a/Engine/Plugins/Online/OnlineSubsystemUtils/Source/OnlineSubsystemUtils/Public/VoipListenerSynthComponent.h b/Engine/Plugins/Online/OnlineSubsystemUtils/Source/OnlineSubsystemUtils/Public/VoipListenerSynthComponent.h
--- a/Engine/Plugins/Online/OnlineSubsystemUtils/Source/OnlineSubsystemUtils/Public/VoipListenerSynthComponent.h
+++ b/Engine/Plugins/Online/OnlineSubsystemUtils/Source/OnlineSubsystemUtils/Public/VoipListenerSynthComponent.h
@@ -7,6 +7,7 @@
 #include "VoicePacketBuffer.h"
 #include "DSP/DynamicsProcessor.h"
 #include "DSP/EQ.h"
+#include "DSP/MultithreadedPatching.h"
 #include "VoipListenerSynthComponent.generated.h"
 
 #define DEBUG_BUFFERING 0
@@ -77,12 +78,17 @@ public:
 	*/
 	void SubmitPacket(void* InBuffer, int32 NumBytes, int64 InStartSample, EVoipStreamDataFormat DataFormat);
 
+	/**
+	 * Optionally connect this VOIP Synth Component to a splitter.
+	 */
+	void ConnectToSplitter(Audio::FPatchMixerSplitter& InSplitter);
+
 	/*
 	 * Thread safe way to get whether this synth component still has buffered packets of audio to play back.
 	 * This function is used by the Voice Engine to clean up voice sounds when they are not playing back audio
 	 * to make space for other sounds to be rendered.
 
-	 * @returns true if this synth component is lollygagging/dawdling/out of audio to play. 
+	 * @returns true if this synth component is out of audio to play. 
 	 */
 	UFUNCTION(BlueprintCallable, Category = "Audio|Voice|Sender")
 	bool IsIdling();
@@ -95,6 +101,8 @@ public:
 	uint64 GetSampleCounter();
 
 private:
+	/** This patch input can optionally be set  */
+	Audio::FPatchInput ExternalSend;
 
 	// This is allocated on OpenPacketStream()
 	TUniquePtr<FVoicePacketBuffer> PacketBuffer;
diff --git a/Engine/Source/Runtime/Online/Voice/Private/Windows/VoiceCaptureWindows.cpp b/Engine/Source/Runtime/Online/Voice/Private/Windows/VoiceCaptureWindows.cpp
--- a/Engine/Source/Runtime/Online/Voice/Private/Windows/VoiceCaptureWindows.cpp
+++ b/Engine/Source/Runtime/Online/Voice/Private/Windows/VoiceCaptureWindows.cpp
@@ -574,6 +574,25 @@ EVoiceCaptureState::Type FVoiceCaptureWindows::GetVoiceData(uint8* OutVoiceBuffe
 		}
 	}
 
+	// If we have any sends for this microphones output, push to them here.
+	if (MicrophoneOutput.Num() > 0)
+	{
+		// Convert our buffer from int16 to float:
+		int16* OutputData = reinterpret_cast<int16*>(OutVoiceBuffer);
+		uint32 NumSamples = OutBytesWritten / sizeof(int16);
+		ConversionBuffer.Reset();
+		// Note: Sample rate is unused for this operation.
+		ConversionBuffer.Append(OutputData, NumSamples, NumInputChannels, 16000);
+		
+		if (NumInputChannels > 1)
+		{
+			// For consistency, mixdown to mono.
+			ConversionBuffer.MixBufferToChannels(1);
+		}
+
+		MicrophoneOutput.PushAudio(ConversionBuffer.GetData(), ConversionBuffer.GetNumSamples());
+	}
+
 	return NewMicState;
 }
 
diff --git a/Engine/Source/Runtime/Online/Voice/Private/Windows/VoiceCaptureWindows.h b/Engine/Source/Runtime/Online/Voice/Private/Windows/VoiceCaptureWindows.h
--- a/Engine/Source/Runtime/Online/Voice/Private/Windows/VoiceCaptureWindows.h
+++ b/Engine/Source/Runtime/Online/Voice/Private/Windows/VoiceCaptureWindows.h
@@ -147,6 +147,9 @@ private:
 	/** Number of channels. */
 	int32 NumInputChannels;
 
+	/** This is used to convert audio from int16 to float to push to any potential connected patch outputs on IVoiceCapture::MicrophoneOutput */
+	Audio::TSampleBuffer<float> ConversionBuffer;
+
 	/**
 	* This buffer is used so that we can detect when the player is speaking
 	* without cutting off what they are saying.
diff --git a/Engine/Source/Runtime/Online/Voice/Public/Interfaces/VoiceCapture.h b/Engine/Source/Runtime/Online/Voice/Public/Interfaces/VoiceCapture.h
--- a/Engine/Source/Runtime/Online/Voice/Public/Interfaces/VoiceCapture.h
+++ b/Engine/Source/Runtime/Online/Voice/Public/Interfaces/VoiceCapture.h
@@ -4,6 +4,7 @@
 
 #include "CoreMinimal.h"
 #include "VoicePackage.h"
+#include "DSP/MultithreadedPatching.h"
 
 class Error;
 
@@ -71,6 +72,11 @@ protected:
 
 	IVoiceCapture() {};
 
+	/**
+	 * Optional patch point for MicrophoneOutput.
+	 */
+	Audio::FPatchSplitter MicrophoneOutput;
+
 public:
 	
 	virtual ~IVoiceCapture() {}
@@ -150,4 +156,12 @@ public:
 
 	/** Dump the state of the voice capture device */
 	virtual void DumpState() const = 0;
+
+	/**
+	 * @returns a new output for this mic's audio.
+	 */
+	Audio::FPatchOutputStrongPtr GetMicrophoneAudio(int32 MaxExpectedLatencyInSamples, float Gain)
+	{
+		return MicrophoneOutput.AddNewPatch(MaxExpectedLatencyInSamples, Gain);
+	}
 };
diff --git a/Engine/Source/Runtime/SignalProcessing/Private/BufferVectorOperations.cpp b/Engine/Source/Runtime/SignalProcessing/Private/BufferVectorOperations.cpp
--- a/Engine/Source/Runtime/SignalProcessing/Private/BufferVectorOperations.cpp
+++ b/Engine/Source/Runtime/SignalProcessing/Private/BufferVectorOperations.cpp
@@ -331,6 +331,63 @@ namespace Audio
 #endif
 	}
 
+	void MixInBufferFast(const AlignedFloatBuffer& InFloatBuffer, AlignedFloatBuffer& BufferToSumTo, const float StartGain, const float EndGain)
+	{
+		MixInBufferFast(InFloatBuffer.GetData(), BufferToSumTo.GetData(), InFloatBuffer.Num(), StartGain, EndGain);
+	}
+
+	void MixInBufferFast(const float* RESTRICT InFloatBuffer, float* RESTRICT BufferToSumTo, int32 NumSamples, const float StartGain, const float EndGain)
+	{
+		checkf(IsAligned<const float*>(InFloatBuffer, 4), TEXT("Memory must be aligned to use vector operations."));
+		checkf(IsAligned<float*>(BufferToSumTo, 4), TEXT("Memory must be aligned to use vector operations."));
+		checkf(NumSamples % 4 == 0, TEXT("Please use a buffer size that is a multiple of 4."));
+
+		const int32 NumIterations = NumSamples / 4;
+
+		if (FMath::IsNearlyEqual(StartGain, EndGain))
+		{
+			// No need to do anything if start and end values are both 0.0
+			if (StartGain == 0.0f)
+			{
+				return;
+			}
+			else
+			{
+				VectorRegister Gain = VectorLoadFloat1(&StartGain);
+
+				for (int32 i = 0; i < NumSamples; i += 4)
+				{
+					VectorRegister Input = VectorLoadAligned(&InFloatBuffer[i]);
+					VectorRegister Output = VectorLoadAligned(&BufferToSumTo[i]);
+
+					Input = VectorMultiply(Input, Gain);
+					Output = VectorAdd(Input, Output);
+					
+					VectorStoreAligned(Output, &BufferToSumTo[i]);
+				}
+			}
+		}
+		else
+		{
+			const float DeltaValue = ((EndGain - StartGain) / NumIterations);
+
+			VectorRegister Gain = VectorLoadFloat1(&StartGain);
+			VectorRegister Delta = VectorLoadFloat1(&DeltaValue);
+
+			for (int32 i = 0; i < NumSamples; i += 4)
+			{
+				VectorRegister Input = VectorLoadAligned(&InFloatBuffer[i]);
+				VectorRegister Output = VectorLoadAligned(&BufferToSumTo[i]);
+				Input = VectorMultiply(Input, Gain);
+				Output = VectorAdd(Input, Output);
+				
+				VectorStoreAligned(Output, &BufferToSumTo[i]);
+
+				Gain = VectorAdd(Gain, Delta);
+			}
+		}
+	}
+
 	/* Subtracts two buffers together element-wise. */
 	void BufferSubtractFast(const AlignedFloatBuffer& InMinuend, const AlignedFloatBuffer& InSubtrahend, AlignedFloatBuffer& OutputBuffer)
 	{
diff --git a/Engine/Source/Runtime/SignalProcessing/Private/MultithreadedPatching.cpp b/Engine/Source/Runtime/SignalProcessing/Private/MultithreadedPatching.cpp
--- a/Engine/Source/Runtime/SignalProcessing/Private/MultithreadedPatching.cpp
+++ b/Engine/Source/Runtime/SignalProcessing/Private/MultithreadedPatching.cpp
@@ -3,6 +3,7 @@
 #include "DSP/MultithreadedPatching.h"
 #include "DSP/BufferVectorOperations.h"
 
+
 namespace Audio
 {
 	TAtomic<int32> FPatchOutput::PatchIDCounter(0);
@@ -12,19 +13,36 @@ namespace Audio
 		, TargetGain(InGain)
 		, PreviousGain(InGain)
 		, PatchID(++PatchIDCounter)
+		, NumAliveInputs(0)
 	{
 
 	}
 
 
+	FPatchOutput::FPatchOutput()
+		: InternalBuffer(0)
+		, TargetGain(0.0f)
+		, PreviousGain(0.0f)
+		, PatchID(INDEX_NONE)
+		, NumAliveInputs(0)
+	{
+	}
+
 	int32 FPatchOutput::PopAudio(float* OutBuffer, int32 NumSamples, bool bUseLatestAudio)
 	{
+		if (IsInputStale())
+		{
+			return -1;
+		}
+
 		if (bUseLatestAudio && InternalBuffer.Num() > (uint32) NumSamples)
 		{
 			InternalBuffer.SetNum(((uint32)NumSamples));
 		}
 
 		int32 PopResult = InternalBuffer.Pop(OutBuffer, NumSamples);
+
+		// Apply gain stage.
 		if (FMath::IsNearlyEqual(TargetGain, PreviousGain))
 		{
 			MultiplyBufferByConstantInPlace(OutBuffer, PopResult, PreviousGain);
@@ -38,8 +56,18 @@ namespace Audio
 		return PopResult;
 	}
 
+	bool FPatchOutput::IsInputStale() const
+	{
+		return NumAliveInputs == 0;
+	}
+
 	int32 FPatchOutput::MixInAudio(float* OutBuffer, int32 NumSamples, bool bUseLatestAudio)
 	{
+		if (IsInputStale())
+		{
+			return -1;
+		}
+
 		MixingBuffer.SetNumUninitialized(NumSamples, false);
 		int32 PopResult = 0;
 		
@@ -67,39 +95,82 @@ namespace Audio
 	}
 
 
-	FPatchInput::FPatchInput(const FPatchOutputPtr& InOutput)
+	FPatchInput::FPatchInput(const FPatchOutputStrongPtr& InOutput)
 		: OutputHandle(InOutput)
+		, PushCallsCounter(0)
 	{
+		if (OutputHandle.IsValid())
+		{
+			OutputHandle->NumAliveInputs++;
+		}
+	}
 
+	FPatchInput::FPatchInput()
+		: PushCallsCounter(0)
+	{
+	}
+
+	FPatchInput::FPatchInput(const FPatchInput& Other)
+		: FPatchInput(Other.OutputHandle)
+	{
+	}
+
+	FPatchInput& FPatchInput::operator=(const FPatchInput& Other)
+	{
+		OutputHandle = Other.OutputHandle;
+		PushCallsCounter = 0;
+		
+		if (OutputHandle.IsValid())
+		{
+			OutputHandle->NumAliveInputs++;
+		}
+
+		return *this;
+	}
+
+	FPatchInput::~FPatchInput()
+	{
+		if (OutputHandle.IsValid())
+		{
+			OutputHandle->NumAliveInputs--;
+		}
 	}
 
 	int32 FPatchInput::PushAudio(const float* InBuffer, int32 NumSamples)
 	{
-		FPatchOutputPtr OutputPtr = OutputHandle.Pin();
-		if (!OutputPtr.IsValid())
+		if (!OutputHandle.IsValid())
 		{
 			return -1;
 		}
 
-		int32 SamplesPushed = OutputPtr->InternalBuffer.Push(InBuffer, NumSamples);
+		int32 SamplesPushed = OutputHandle->InternalBuffer.Push(InBuffer, NumSamples);
+
+		// Every so often, we check to see if the output handle has been destroyed and clean it up.
+		static const int32 NumPushCallsUntilCleanupCheck = 256;
+		
+		PushCallsCounter = (PushCallsCounter + 1) % NumPushCallsUntilCleanupCheck;
+		if (PushCallsCounter == 0 && OutputHandle.IsUnique())
+		{
+			// Delete the output.
+			OutputHandle.Reset();
+		}
 
 		return SamplesPushed;
 	}
 
 	void FPatchInput::SetGain(float InGain)
 	{
-		FPatchOutputPtr OutputPtr = OutputHandle.Pin();
-		if (!OutputPtr.IsValid())
+		if (!OutputHandle.IsValid())
 		{
 			return;
 		}
 
-		OutputPtr->TargetGain = InGain;
+		OutputHandle->TargetGain = InGain;
 	}
 
 	bool FPatchInput::IsOutputStillActive()
 	{
-		return OutputHandle.IsValid();
+		return OutputHandle.IsUnique() || OutputHandle.IsValid();
 	}
 
 	FPatchMixer::FPatchMixer()
@@ -107,37 +178,37 @@ namespace Audio
 	}
 
 
-	FPatchInput FPatchMixer::AddNewPatch(int32 MaxLatencyInSamples, float InGain)
+	FPatchInput FPatchMixer::AddNewInput(int32 MaxLatencyInSamples, float InGain)
 	{
-		FScopeLock ScopeLock(&PendingNewPatchesCriticalSection);
+		FScopeLock ScopeLock(&PendingNewInputsCriticalSection);
 
-		int32 NewPatchIndex = PendingNewPatches.Emplace(new FPatchOutput(MaxLatencyInSamples, InGain));
-		return FPatchInput(PendingNewPatches[NewPatchIndex]);
+		int32 NewPatchIndex = PendingNewInputs.Emplace(new FPatchOutput(MaxLatencyInSamples, InGain));
+		return FPatchInput(PendingNewInputs[NewPatchIndex]);
 	}
 
 	void FPatchMixer::RemovePatch(const FPatchInput& PatchInput)
 	{
-		FPatchOutputPtr OutputPtr = PatchInput.OutputHandle.Pin();
-
 		// If the output is already disconnected, early exit.
-		if (!OutputPtr.IsValid())
+		if (!PatchInput.OutputHandle.IsValid())
 		{
 			return;
 		}
 
-		FScopeLock ScopeLock(&PatchDeletionCriticalSection);
-		DisconnectedPatches.Add(OutputPtr->PatchID);
+		FScopeLock ScopeLock(&InputDeletionCriticalSection);
+		DisconnectedInputs.Add(PatchInput.OutputHandle->PatchID);
 	}
 
 	int32 FPatchMixer::PopAudio(float* OutBuffer, int32 OutNumSamples, bool bUseLatestAudio)
 	{
+		FScopeLock ScopeLock(&CurrentPatchesCriticalSection);
+
 		CleanUpDisconnectedPatches();
 		ConnectNewPatches();
 
 		FMemory::Memzero(OutBuffer, OutNumSamples * sizeof(float));
 		int32 MaxPoppedSamples = 0;
 
-		for (FPatchOutputPtr& OutputPtr : CurrentPatches)
+		for (FPatchOutputStrongPtr& OutputPtr : CurrentInputs)
 		{
 			const int32 NumPoppedSamples = OutputPtr->MixInAudio(OutBuffer, OutNumSamples, bUseLatestAudio);
 			MaxPoppedSamples = FMath::Max(NumPoppedSamples, MaxPoppedSamples);
@@ -146,68 +217,239 @@ namespace Audio
 		return MaxPoppedSamples;
 	}
 
-	void FPatchMixer::ConnectNewPatches()
+	int32 FPatchMixer::Num()
 	{
-		// If AddNewPatch is called in a separate thread, wait until the next PopAudio call to do this work.
-		if (PendingNewPatchesCriticalSection.TryLock())
+		FScopeLock ScopeLock(&CurrentPatchesCriticalSection);
+		return CurrentInputs.Num();
+	}
+
+	int32 FPatchMixer::MaxNumberOfSamplesThatCanBePopped()
+	{
+		FScopeLock ScopeLock(&CurrentPatchesCriticalSection);
+		ConnectNewPatches();
+
+		// Iterate through our inputs and see which input has the least audio buffered.
+		uint32 SmallestNumSamplesBuffered = TNumericLimits<uint32>::Max();
+
+		for (FPatchOutputStrongPtr& Output : CurrentInputs)
 		{
-			// Todo: convert this to move semantics to avoid copying the shared pointer around.
-			for (FPatchOutputPtr& Patch : PendingNewPatches)
+			if (Output.IsValid())
 			{
-				CurrentPatches.Add(Patch);
+				SmallestNumSamplesBuffered = FMath::Min(SmallestNumSamplesBuffered, Output->InternalBuffer.Num());
 			}
+		}
+
+		if (SmallestNumSamplesBuffered == TNumericLimits<uint32>::Max())
+		{
+			return -1;
+		}
+		else
+		{
+			// If this check is hit, we need to either change this function to return an int64 or find a different way to notify the caller that all outputs have been disconeccted.
+			check(SmallestNumSamplesBuffered <= ((uint32)TNumericLimits<int32>::Max()));
+			return SmallestNumSamplesBuffered;
+		}
+	}
 
-			PendingNewPatches.Reset();
-			PendingNewPatchesCriticalSection.Unlock();
+	void FPatchMixer::ConnectNewPatches()
+	{
+		FScopeLock ScopeLock(&PendingNewInputsCriticalSection);
+
+		// If AddNewPatch is called in a separate thread, wait until the next PopAudio call to do this work.
+		// Todo: convert this to move semantics to avoid copying the shared pointer around.
+		for (FPatchOutputStrongPtr& Patch : PendingNewInputs)
+		{
+			CurrentInputs.Add(Patch);
 		}
+
+		PendingNewInputs.Reset();
 	}
 
 	void FPatchMixer::CleanUpDisconnectedPatches()
 	{
-		if (PatchDeletionCriticalSection.TryLock())
+		FScopeLock PendingInputDeletionScopeLock(&InputDeletionCriticalSection);
+
+		// Iterate through all of the PatchIDs we need to clean up.
+		for (const int32& PatchID : DisconnectedInputs)
 		{
-			// Iterate through all of the PatchIDs we need to clean up.
-			for (const int32& PatchID : DisconnectedPatches)
-			{
-				bool bPatchRemoved = false;
+			bool bInputRemoved = false;
 
-				// First, make sure that the patch isn't in the pending new patchs we haven't added yet:
+			// First, make sure that the patch isn't in the pending new patchs we haven't added yet:
+			{
+				FScopeLock PendingNewInputsScopeLock(&PendingNewInputsCriticalSection);
+				for (int32 Index = 0; Index < PendingNewInputs.Num(); Index++)
 				{
-					FScopeLock ScopeLock(&PendingNewPatchesCriticalSection);
-					for (int32 Index = 0; Index < PendingNewPatches.Num(); Index++)
+					checkSlow(CurrentInputs[Index].IsValid());
+
+					if (PatchID == PendingNewInputs[Index]->PatchID)
 					{
-						checkSlow(CurrentPatches[Index].IsValid());
-
-						if (PatchID == PendingNewPatches[Index]->PatchID)
-						{
-							PendingNewPatches.RemoveAtSwap(Index);
-							bPatchRemoved = true;
-							break;
-						}
+						PendingNewInputs.RemoveAtSwap(Index);
+						bInputRemoved = true;
+						break;
 					}
 				}
+			}
+
+			if (bInputRemoved)
+			{
+				continue;
+			}
 
-				if (bPatchRemoved)
+			// Next, we check out current patchs.
+			for (int32 Index = 0; Index < CurrentInputs.Num(); Index++)
+			{
+				checkSlow(CurrentInputs[Index].IsValid());
+
+				if (PatchID == CurrentInputs[Index]->PatchID)
 				{
-					continue;
+					CurrentInputs.RemoveAtSwap(Index);
+					break;
 				}
+			}
+		}
 
-				// Next, we check out current patchs.
-				for (int32 Index = 0; Index < CurrentPatches.Num(); Index++)
-				{
-					checkSlow(CurrentPatches[Index].IsValid());
+		DisconnectedInputs.Reset();
+	}
 
-					if (PatchID == CurrentPatches[Index]->PatchID)
-					{
-						CurrentPatches.RemoveAtSwap(Index);
-						break;
-					}
-				}
+	FPatchSplitter::FPatchSplitter()
+	{
+	}
+
+	FPatchSplitter::~FPatchSplitter()
+	{
+	}
+
+	Audio::FPatchOutputStrongPtr FPatchSplitter::AddNewPatch(int32 MaxLatencyInSamples, float InGain)
+	{
+		// Allocate a new FPatchOutput, then store a weak pointer to it in our PendingOutputs array to be added in our next call to PushAudio.
+		FPatchOutputStrongPtr StrongOutputPtr = MakeShareable(new FPatchOutput(MaxLatencyInSamples * 2, InGain));
+
+		{
+			FScopeLock ScopeLock(&PendingOutputsCriticalSection);
+			PendingOutputs.Emplace(StrongOutputPtr);
+		}
+
+		return StrongOutputPtr;
+	}
+
+	int32 FPatchSplitter::Num()
+	{
+		FScopeLock ScopeLock(&ConnectedOutputsCriticalSection);
+		AddPendingPatches();
+		return ConnectedOutputs.Num();
+	}
+
+	int32 FPatchSplitter::MaxNumberOfSamplesThatCanBePushed()
+	{
+		FScopeLock ScopeLock(&ConnectedOutputsCriticalSection);
+		AddPendingPatches();
+
+		// Iterate over our outputs and get the smallest remainder of all of our circular buffers.
+		uint32 SmallestRemainder = TNumericLimits<uint32>::Max();
+
+		for (FPatchInput& Input : ConnectedOutputs)
+		{
+			if (Input.OutputHandle.IsValid())
+			{
+				SmallestRemainder = FMath::Min(SmallestRemainder, Input.OutputHandle->InternalBuffer.Remainder());
 			}
+		}
 
-			DisconnectedPatches.Reset();
+		if (SmallestRemainder == TNumericLimits<uint32>::Max())
+		{
+			return -1;
+		}
+		else
+		{
+			// If we hit this check, we need to either return an int64 or use some other method to notify the caller that all outputs are disconnected.
+			check(SmallestRemainder <= ((uint32)TNumericLimits<int32>::Max()));
+			return SmallestRemainder;
+		}
+	}
 
-			PatchDeletionCriticalSection.Unlock();
+	void FPatchSplitter::AddPendingPatches()
+	{
+		FScopeLock ScopeLock(&PendingOutputsCriticalSection);
+		ConnectedOutputs.Append(PendingOutputs);
+		PendingOutputs.Reset();
+	}
+
+	int32 FPatchSplitter::PushAudio(const float* InBuffer, int32 InNumSamples)
+	{
+		AddPendingPatches();
+
+		FScopeLock ScopeLock(&ConnectedOutputsCriticalSection);
+
+		int32 MinimumSamplesPushed = TNumericLimits<int32>::Max();
+
+		// Iterate through our array of connected outputs from the end, removing destroyed outputs as we go.
+		for (int32 Index = ConnectedOutputs.Num() - 1; Index >= 0; Index--)
+		{
+			int32 NumSamplesPushed = ConnectedOutputs[Index].PushAudio(InBuffer, InNumSamples);
+			if (NumSamplesPushed >= 0)
+			{
+				MinimumSamplesPushed = FMath::Min(MinimumSamplesPushed, NumSamplesPushed);
+			}
+			else
+			{
+				// If this output has been destroyed, remove it from our array of connected outputs.
+				ConnectedOutputs.RemoveAtSwap(Index);
+			}
+		}
+
+		// If we weren't able to push audio to any of our outputs, return -1.
+		if (MinimumSamplesPushed == TNumericLimits<int32>::Max())
+		{
+			MinimumSamplesPushed = -1;
+		}
+
+		return MinimumSamplesPushed;
+	}
+
+	FPatchMixerSplitter::FPatchMixerSplitter()
+	{
+	}
+
+	FPatchMixerSplitter::~FPatchMixerSplitter()
+	{
+	}
+
+	Audio::FPatchOutputStrongPtr FPatchMixerSplitter::AddNewOutput(int32 MaxLatencyInSamples, float InGain)
+	{
+		return Splitter.AddNewPatch(MaxLatencyInSamples, InGain);
+	}
+
+	Audio::FPatchInput FPatchMixerSplitter::AddNewInput(int32 MaxLatencyInSamples, float InGain)
+	{
+		return Mixer.AddNewInput(MaxLatencyInSamples, InGain);
+	}
+
+	void FPatchMixerSplitter::RemovePatch(const FPatchInput& TapInput)
+	{
+		Mixer.RemovePatch(TapInput);
+	}
+
+	void FPatchMixerSplitter::ProcessAudio()
+	{
+		int32 NumSamplesToForward = FMath::Min(Mixer.MaxNumberOfSamplesThatCanBePopped(), Splitter.MaxNumberOfSamplesThatCanBePushed());
+		
+		if (NumSamplesToForward <= 0)
+		{
+			// Likely there are either no inputs or no outputs connected, or one of the inputs has not pushed any audio yet. Early exit.
+			return;
 		}
+
+		IntermediateBuffer.Reset();
+		IntermediateBuffer.AddUninitialized(NumSamplesToForward);
+
+		// Mix down inputs:
+		int32 PopResult = Mixer.PopAudio(IntermediateBuffer.GetData(), NumSamplesToForward, false);
+		check(PopResult == NumSamplesToForward);
+		
+		OnProcessAudio(TArrayView<const float>(IntermediateBuffer));
+
+		// Push audio to outputs:
+		int32 PushResult = Splitter.PushAudio(IntermediateBuffer.GetData(), NumSamplesToForward);
+		check(PushResult == NumSamplesToForward);
 	}
 }
diff --git a/Engine/Source/Runtime/SignalProcessing/Public/DSP/BufferVectorOperations.h b/Engine/Source/Runtime/SignalProcessing/Public/DSP/BufferVectorOperations.h
--- a/Engine/Source/Runtime/SignalProcessing/Public/DSP/BufferVectorOperations.h
+++ b/Engine/Source/Runtime/SignalProcessing/Public/DSP/BufferVectorOperations.h
@@ -68,6 +68,8 @@ namespace Audio
 	/** Takes buffer InFloatBuffer, optionally multiplies it by Gain, and adds it to BufferToSumTo. */
 	SIGNALPROCESSING_API void MixInBufferFast(const AlignedFloatBuffer& InFloatBuffer, AlignedFloatBuffer& BufferToSumTo, const float Gain);
 	SIGNALPROCESSING_API void MixInBufferFast(const float* RESTRICT InFloatBuffer, float* RESTRICT BufferToSumTo, int32 NumSamples, const float Gain);
+	SIGNALPROCESSING_API void MixInBufferFast(const AlignedFloatBuffer& InFloatBuffer, AlignedFloatBuffer& BufferToSumTo, const float StartGain, const float EndGain);
+	SIGNALPROCESSING_API void MixInBufferFast(const float* RESTRICT InFloatBuffer, float* RESTRICT BufferToSumTo, int32 NumSamples, const float StartGain, const float EndGain);
 	SIGNALPROCESSING_API void MixInBufferFast(const AlignedFloatBuffer& InFloatBuffer, AlignedFloatBuffer& BufferToSumTo);
 	SIGNALPROCESSING_API void MixInBufferFast(const float* RESTRICT InFloatBuffer, float* RESTRICT BufferToSumTo, int32 NumSamples);
 
diff --git a/Engine/Source/Runtime/SignalProcessing/Public/DSP/MultithreadedPatching.h b/Engine/Source/Runtime/SignalProcessing/Public/DSP/MultithreadedPatching.h
--- a/Engine/Source/Runtime/SignalProcessing/Public/DSP/MultithreadedPatching.h
+++ b/Engine/Source/Runtime/SignalProcessing/Public/DSP/MultithreadedPatching.h
@@ -9,21 +9,42 @@
 namespace Audio
 {
 	/**
-	 * This is a buffer that is owned by the FPatchMixer.
+	 * This class can be thought of as an output for a single constructed instance of FPatchInput.
+	 * Each FPatchOutput can only be connected to one FPatchInput. To route multiple outputs, see FPatchSplitter.
+	 * To route multiple inputs, see FPatchMixer.
+	 *
+	 * Example usage:
+	 * 
+	 * FPatchOutputStrongPtr NewOutput(new FPatchOutput(4096));
+	 * FPatchInput NewInput(NewOutput);
+	 *
+	 * // On one thread, push audio to the output:
+	 * NewInput.PushAudio(AudioBufferPtr, AudioBufferNumSamples);
+	 *
+	 * // and on a seperate thread, retrieve the audio:
+	 * NewOutput->PopAudio(OutAudioBufferPtr, AudioBufferNumSamples);
+	 * 
 	 */
 	struct SIGNALPROCESSING_API FPatchOutput
 	{
 	public:
 		FPatchOutput(int32 MaxCapacity, float InGain = 1.0f);
 
-		/** Copies the minimum of NumSamples or however many samples are available into OutBuffer. Returns the number of samples copied. */
+		/** The default constructor will result in an uninitialized, disconnected patch point, which  */
+		FPatchOutput();
+
+		/** Copies the minimum of NumSamples or however many samples are available into OutBuffer. Returns the number of samples copied, or -1 if this output's corresponding input has been destroyed. */
 		int32 PopAudio(float* OutBuffer, int32 NumSamples, bool bUseLatestAudio);
 
 		/** Sums the minimum of NumSamples or however many samples are available into OutBuffer. Returns the number of samples summed into OutBuffer. */
 		int32 MixInAudio(float* OutBuffer, int32 NumSamples, bool bUseLatestAudio);
 
+		/** Returns true if the input for this patch has been destroyed. */
+		bool IsInputStale() const;
+
 		friend class FPatchInput;
 		friend class FPatchMixer;
+		friend class FPatchSplitter;
 	private:
 
 		// Internal buffer.
@@ -39,11 +60,15 @@ namespace Audio
 		// This is used to breadcrumb the FPatchOutput when we want to delete it.
 		int32 PatchID;
 
+		// Counter that is incremented/decremented to allow FPatchInput to be copied around safely.
+		int32 NumAliveInputs;
+
 		static TAtomic<int32> PatchIDCounter;
 	};
 
 	/** Patch outputs are owned by the FPatchMixer, and are pinned by the FPatchInput. */
-	typedef TSharedPtr<FPatchOutput, ESPMode::ThreadSafe> FPatchOutputPtr;
+	typedef TSharedPtr<FPatchOutput, ESPMode::ThreadSafe> FPatchOutputStrongPtr;
+	typedef TWeakPtr<FPatchOutput, ESPMode::ThreadSafe> FPatchOutputWeakPtr;
 
 	/**
 	 * Handle to a patch. Should only be used by a single thread.
@@ -51,7 +76,15 @@ namespace Audio
 	class SIGNALPROCESSING_API FPatchInput
 	{
 	public:
-		FPatchInput(const FPatchOutputPtr& InOutput);
+		/** PatchInputs can only be created from explicit outputs. */
+		FPatchInput(const FPatchOutputStrongPtr& InOutput);
+		FPatchInput(const FPatchInput& Other);
+		FPatchInput& operator=(const FPatchInput& Other);
+
+		/** Default constructed FPatchInput instances will always return -1 for PushAudio and true for IsOutputStillActive. */
+		FPatchInput();
+
+		~FPatchInput();
 
 		/** pushes audio from InBuffer to the corresponding FPatchOutput.
 		 *  Returns how many samples were able to be pushed, or -1 if the output was disconnected.
@@ -64,10 +97,14 @@ namespace Audio
 		bool IsOutputStillActive();
 
 		friend class FPatchMixer;
+		friend class FPatchSplitter;
 
 	private:
 		/** Weak pointer to our destination buffer. */
-		TWeakPtr<FPatchOutput, ESPMode::ThreadSafe> OutputHandle;
+		FPatchOutputStrongPtr OutputHandle;
+
+		/** Counter of the number of push calls until we check OutputHandle.IsUnique() and possibly clean up the output. */
+		int32 PushCallsCounter;
 	};
 
 	/**
@@ -82,7 +119,7 @@ namespace Audio
 		FPatchMixer();
 
 		/** Adds a new input to the tap collector. Calling this is thread safe, but individual instances of FPatchInput are only safe to be used from one thread. */
-		FPatchInput AddNewPatch(int32 MaxLatencyInSamples, float InGain);
+		FPatchInput AddNewInput(int32 MaxLatencyInSamples, float InGain);
 
 		/** Removes a tap from the tap collector. Calling this is thread safe, though FPatchOutput will likely not be deleted until the next call of PopAudio. */
 		void RemovePatch(const FPatchInput& TapInput);
@@ -90,6 +127,12 @@ namespace Audio
 		/** Mixes all inputs into a single buffer. This should only be called from a single thread. Returns the number of non-silent samples popped to OutBuffer. */
 		int32 PopAudio(float* OutBuffer, int32 OutNumSamples, bool bUseLatestAudio);
 
+		/** This returns the number of inputs currently connected to this patch mixer. Thread safe, but blocks for PopAudio. */
+		int32 Num();
+
+		/** This function call gets the maximum number of samples that's safe to pop, based on the thread with the least amount of samples buffered. Thread safe, but blocks for PopAudio. */
+		int32 MaxNumberOfSamplesThatCanBePopped();
+
 	private:
 		/** Called within PopAudio. Flushes the PendingNewPatches array into CurrentPatches. During this function, AddNewPatch is blocked. */
 		void ConnectNewPatches();
@@ -98,16 +141,97 @@ namespace Audio
 		void CleanUpDisconnectedPatches();
 
 		/** New taps are added here in AddNewPatch, and then are moved to CurrentPatches in ConnectNewPatches. */
-		TArray<FPatchOutputPtr> PendingNewPatches;
+		TArray<FPatchOutputStrongPtr> PendingNewInputs;
 		/** Contended by AddNewPatch, ConnectNewPatches and CleanUpDisconnectedTaps. */
-		FCriticalSection PendingNewPatchesCriticalSection;
+		FCriticalSection PendingNewInputsCriticalSection;
 
 		/** Patch IDs of individual audio taps that will be removed on the next call of CleanUpDisconnectedPatches. */
-		TArray<int32> DisconnectedPatches;
+		TArray<int32> DisconnectedInputs;
 		/** Contended by RemoveTap, AddNewPatch, and ConnectNewPatches. */
-		FCriticalSection PatchDeletionCriticalSection;
+		FCriticalSection InputDeletionCriticalSection;
 
 		/** Only accessed within PopAudio. Indirect array of taps that are mixed in during PopAudio. */
-		TArray<FPatchOutputPtr> CurrentPatches;
+		TArray<FPatchOutputStrongPtr> CurrentInputs;
+		FCriticalSection CurrentPatchesCriticalSection;
+	};
+
+	/**
+	 * This class is used to post audio from one source to multiple threads.
+	 * This class is SPMC: multiple threads can call FPatchOutputStrongPtr->PopAudio safely,
+	 * but only one thread can call PushAudio.
+	 */
+	class SIGNALPROCESSING_API FPatchSplitter
+	{
+	public:
+		/** Constructor. */
+		FPatchSplitter();
+
+		/** The destructor will mark every still connected FPatchOutput as stale. */
+		~FPatchSplitter();
+
+		/**
+		 * Adds a new output. Calling this is thread safe, but individual instances of FPatchOutput are only safe to be used from one thread.
+		 * the returned FPatchOutputPtr can be safely destroyed at any point.
+		 */
+		FPatchOutputStrongPtr AddNewPatch(int32 MaxLatencyInSamples, float InGain);
+
+		/** This call pushes audio to all outputs connected to this splitter. Only should be called from one thread. */
+		int32 PushAudio(const float* InBuffer, int32 InNumSamples);
+
+		/** This returns the number of outputs currently connected to this patch splitter. Thread safe, but blocks for PushAudio. */
+		int32 Num();
+
+		/** This function call gets the maximum number of samples that's safe to push. Thread safe, but blocks for PushAudio. */
+		int32 MaxNumberOfSamplesThatCanBePushed();
+
+	private:
+		// Called from PushAudio(). 
+		void AddPendingPatches();
+
+		TArray<FPatchInput> PendingOutputs;
+		FCriticalSection PendingOutputsCriticalSection;
+
+		TArray<FPatchInput> ConnectedOutputs;
+		FCriticalSection ConnectedOutputsCriticalSection;
+	};
+
+	/**
+	 * This class is used to mix multiple inputs from disparate threads to a single mixdown and deliver that mixdown to multiple outputs.
+	 * This class is MPCMC, but only one thread can and should call ProcessAudio().
+	 */
+	class AUDIOMIXER_API FPatchMixerSplitter
+	{
+	public:
+		/** Constructor. */
+		FPatchMixerSplitter();
+
+		/** The destructor will mark every still connected FPatchOutput as stale. */
+		virtual ~FPatchMixerSplitter();
+
+		/**
+		 * Adds a new output. Calling this is thread safe, but individual instances of FPatchOutput are only safe to be used from one thread.
+		 * the returned FPatchOutputPtr can be safely destroyed at any point.
+		 */
+		FPatchOutputStrongPtr AddNewOutput(int32 MaxLatencyInSamples, float InGain);
+
+		/** Adds a new input to the tap collector. Calling this is thread safe, but individual instances of FPatchInput are only safe to be used from one thread. */
+		FPatchInput AddNewInput(int32 MaxLatencyInSamples, float InGain);
+
+		/** Removes a tap from the tap collector. Calling this is thread safe, though FPatchOutput will likely not be deleted until the next call of PopAudio. */
+		void RemovePatch(const FPatchInput& TapInput);
+
+		/** Mixes audio from all inputs and pushes it to all outputs. Should be called regularly. */
+		void ProcessAudio();
+
+	protected:
+		/** This class can be subclassed with OnProcessAudio overridden. */
+		virtual void OnProcessAudio(TArrayView<const float> InAudio) {}
+
+	private:
+		FPatchMixer Mixer;
+		FPatchSplitter Splitter;
+
+		/** This buffer is used to pop audio from our Mixer and push to our splitter. */
+		AlignedFloatBuffer IntermediateBuffer;
 	};
 }
diff --git a/Engine/Source/Runtime/SignalProcessing/Public/SampleBuffer.h b/Engine/Source/Runtime/SignalProcessing/Public/SampleBuffer.h
--- a/Engine/Source/Runtime/SignalProcessing/Public/SampleBuffer.h
+++ b/Engine/Source/Runtime/SignalProcessing/Public/SampleBuffer.h
@@ -253,8 +253,28 @@ namespace Audio
 			SampleDuration = (float)NumFrames / SampleRate;
 		}
 
+		// Overload of Append that also sets the number of channels and sample rate.
+		template<class OtherSampleType>
+		void Append(const OtherSampleType* InputBuffer, int32 InNumSamples, int32 InNumChannels, int32 InSampleRate)
+		{
+			NumChannels = InNumChannels;
+			SampleRate = InSampleRate;
+
+			Append(InputBuffer, InNumSamples);
+		}
+
 		~TSampleBuffer() {};
 
+		void Reset()
+		{
+			RawPCMData.Reset();
+			NumSamples = 0;
+			NumFrames = 0;
+			NumChannels = 0;
+			SampleRate = 0.0f;
+			SampleDuration = 0.0f;
+		}
+
 		// Gets the raw PCM data of the sound wave
 		FORCEINLINE const SampleType* GetData() const
 		{
