[NO NON-PII COMMIT MESSAGE]

diff --git a/Engine/Source/Runtime/SignalProcessing/Private/BufferVectorOperations.cpp b/Engine/Source/Runtime/SignalProcessing/Private/BufferVectorOperations.cpp
--- a/Engine/Source/Runtime/SignalProcessing/Private/BufferVectorOperations.cpp
+++ b/Engine/Source/Runtime/SignalProcessing/Private/BufferVectorOperations.cpp
@@ -468,35 +468,6 @@ namespace Audio
 		}
 	}
 
-	SIGNALPROCESSING_API void MixInBufferFast(const AlignedFloatBuffer& InFloatBuffer, AlignedFloatBuffer& BufferToSumTo, const float StartGain, const float EndGain)
-	{
-		MixInBufferFast(InFloatBuffer.GetData(), BufferToSumTo.GetData(), BufferToSumTo.Num(), StartGain, EndGain);
-	}
-
-	SIGNALPROCESSING_API void MixInBufferFast(const float* RESTRICT InFloatBuffer, float* RESTRICT BufferToSumTo, int32 NumSamples, const float StartGain, const float EndGain)
-	{
-		checkf(IsAligned<const float*>(InFloatBuffer, 4), TEXT("Memory must be aligned to use vector operations."));
-		checkf(IsAligned<float*>(BufferToSumTo, 4), TEXT("Memory must be aligned to use vector operations."));
-		checkf(NumSamples % 4 == 0, TEXT("Please use a buffer size that is a multiple of 4."));
-
-		const int32 NumIterations = NumSamples / 4;
-
-		const float DeltaValue = ((EndGain - StartGain) / NumIterations);
-
-		VectorRegister Gain = VectorLoadFloat1(&StartGain);
-		const VectorRegister Delta = VectorLoadFloat1(&DeltaValue);
-
-		for (int32 i = 0; i < NumSamples; i += 4)
-		{
-			VectorRegister Output = VectorLoadAligned(&BufferToSumTo[i]);
-			VectorRegister Input = VectorLoadAligned(&InFloatBuffer[i]);
-			Output = VectorMultiplyAdd(Input, Gain, Output);
-			VectorStoreAligned(Output, &BufferToSumTo[i]);
-
-			Gain = VectorAdd(Gain, Delta);
-		}
-	}
-
 	void SumBuffers(const AlignedFloatBuffer& InFloatBuffer1, const AlignedFloatBuffer& InFloatBuffer2, AlignedFloatBuffer& OutputBuffer)
 	{
 		checkf(InFloatBuffer1.Num() == InFloatBuffer2.Num(), TEXT("Input buffers must be equal length"));
diff --git a/Engine/Source/Runtime/SignalProcessing/Public/DSP/MultithreadedPatching.h b/Engine/Source/Runtime/SignalProcessing/Public/DSP/MultithreadedPatching.h
--- a/Engine/Source/Runtime/SignalProcessing/Public/DSP/MultithreadedPatching.h
+++ b/Engine/Source/Runtime/SignalProcessing/Public/DSP/MultithreadedPatching.h
@@ -199,7 +199,7 @@ namespace Audio
 	 * This class is used to mix multiple inputs from disparate threads to a single mixdown and deliver that mixdown to multiple outputs.
 	 * This class is MPCMC, but only one thread can and should call ProcessAudio().
 	 */
-	class AUDIOMIXER_API FPatchMixerSplitter
+	class SIGNALPROCESSING_API FPatchMixerSplitter
 	{
 	public:
 		/** Constructor. */
