[NO NON-PII COMMIT MESSAGE]

diff --git a/Engine/Source/Runtime/Online/Voice/Private/Windows/VoiceCaptureWindows.h b/Engine/Source/Runtime/Online/Voice/Private/Windows/VoiceCaptureWindows.h
--- a/Engine/Source/Runtime/Online/Voice/Private/Windows/VoiceCaptureWindows.h
+++ b/Engine/Source/Runtime/Online/Voice/Private/Windows/VoiceCaptureWindows.h
@@ -9,6 +9,7 @@
 #include "DSP/EnvelopeFollower.h"
 #include "Containers/Ticker.h"
 #include "Windows/WindowsHWrapper.h"
+#include "SampleBuffer.h"
 
 typedef struct IDirectSound8 *LPDIRECTSOUND8;
 
diff --git a/Engine/Source/Runtime/SignalProcessing/Private/VoiceProcessing.cpp b/Engine/Source/Runtime/SignalProcessing/Private/VoiceProcessing.cpp
--- a/Engine/Source/Runtime/SignalProcessing/Private/VoiceProcessing.cpp
+++ b/Engine/Source/Runtime/SignalProcessing/Private/VoiceProcessing.cpp
@@ -363,7 +363,7 @@ namespace Audio
 
 	Audio::FPatchInput FAcousticEchoCancellation::AddNewSignalPatch(int32 ExpectedLatency, float Gain /*= 1.0f*/)
 	{
-		return PatchMixer.AddNewPatch(ExpectedLatency, Gain);
+		return PatchMixer.AddNewInput(ExpectedLatency, Gain);
 	}
 
 	void FAcousticEchoCancellation::RemoveSignalPatch(const FPatchInput& Patch)
