Preventing overflow in audio output by clamping before converting from float to PCM.

Also adding an volume scalar cvar that is applied on float->int16 conversion.

diff --git a/Engine/Source/Runtime/AudioMixerCore/Private/AudioMixer.cpp b/Engine/Source/Runtime/AudioMixerCore/Private/AudioMixer.cpp
--- a/Engine/Source/Runtime/AudioMixerCore/Private/AudioMixer.cpp
+++ b/Engine/Source/Runtime/AudioMixerCore/Private/AudioMixer.cpp
@@ -1,6 +1,7 @@
 // Copyright 1998-2019 Epic Games, Inc. All Rights Reserved.
 
 #include "AudioMixer.h"
+#include "DSP/BufferVectorOperations.h"
 #include "HAL/RunnableThread.h"
 #include "Misc/ConfigCacheIni.h"
 #include "HAL/ThreadSafeCounter.h"
@@ -64,6 +65,13 @@ FAutoConsoleVariableRef CVarUnderrunTimeout(
 	TEXT("Amount of time to wait for the render thread to generate the next buffer before submitting an underrun buffer. \n"),
 	ECVF_Default);
 
+static float LinearGainScalarForFinalOututCVar = 1.0f;
+FAutoConsoleVariableRef LinearGainScalarForFinalOutut(
+	TEXT("au.LinearGainScalarForFinalOutut"),
+	LinearGainScalarForFinalOututCVar,
+	TEXT("Linear gain scalar applied to the final float buffer to allow for hotfixable mitigation of clipping \n")
+	TEXT("Default is 1.0f \n"),
+	ECVF_Default);
 
 namespace Audio
 {
@@ -177,17 +185,28 @@ namespace Audio
 
 		switch (DataFormat)
 		{
-			// Doesn't do anything...
 		case EAudioMixerStreamDataFormat::Float:
-			break;
+		{
+			if (!FMath::IsNearlyEqual(LinearGainScalarForFinalOututCVar, 1.0f))
+			{
+				MultiplyBufferByConstantInPlace(Buffer, LinearGainScalarForFinalOututCVar);
+			}
+			BufferRangeClampFast(Buffer, -1.0f, 1.0f);
+		}
+		break;
 
 		case EAudioMixerStreamDataFormat::Int16:
 		{
 			int16* BufferInt16 = (int16*)FormattedBuffer.GetData();
 			const int32 NumSamples = Buffer.Num();
+
+			const float ConversionScalar = LinearGainScalarForFinalOututCVar * 32767.0f;
+			MultiplyBufferByConstantInPlace(Buffer, ConversionScalar);
+			BufferRangeClampFast(Buffer, -32767.0f, 32767.0f);
+
 			for (int32 i = 0; i < NumSamples; ++i)
 			{
-				BufferInt16[i] = (int16)(Buffer[i] * 32767.0f);
+				BufferInt16[i] = (int16)Buffer[i];
 			}
 		}
 		break;
diff --git a/Engine/Source/Runtime/SignalProcessing/Private/BufferVectorOperations.cpp b/Engine/Source/Runtime/SignalProcessing/Private/BufferVectorOperations.cpp
--- a/Engine/Source/Runtime/SignalProcessing/Private/BufferVectorOperations.cpp
+++ b/Engine/Source/Runtime/SignalProcessing/Private/BufferVectorOperations.cpp
@@ -35,6 +35,38 @@ namespace Audio
 		}
 	}
 
+	/* Clamps values in the buffer to be between InMinValue and InMaxValue */
+	void BufferRangeClampFast(AlignedFloatBuffer& InOutBuffer, float InMinValue, float InMaxValue)
+	{
+		return BufferRangeClampFast(InOutBuffer.GetData(), InOutBuffer.Num(), InMinValue, InMaxValue);
+	}
+
+	/* Clamps values in the buffer to be between InMinValue and InMaxValue */
+	void BufferRangeClampFast(float* RESTRICT InOutBuffer, const int32 InNum, float InMinValue, float InMaxValue)
+	{
+		checkf(InNum >= 4, TEXT("Buffer must have atleast 4 elements."));
+		checkf(0 == (InNum % 4), TEXT("Buffer length be a multiple of 4."));
+		checkf(IsAligned<float*>(InOutBuffer, AUDIO_SIMD_FLOAT_ALIGNMENT), TEXT("Memory must be aligned to use vector operations."));
+
+		const VectorRegister VMinVal = MakeVectorRegister(InMinValue, InMinValue, InMinValue, InMinValue);
+		const VectorRegister VMaxVal = MakeVectorRegister(InMaxValue, InMaxValue, InMaxValue, InMaxValue);
+
+		for (int32 i = 0; i < InNum; i += 4)
+		{
+			VectorRegister VInOut = VectorLoadAligned(&InOutBuffer[i]);
+
+			// Create masks to flag elements outside of range.
+			VectorRegister MinMask = VectorCompareLT(VInOut, VMinVal);
+			VectorRegister MaxMask = VectorCompareGT(VInOut, VMaxVal);
+
+			// Choose between range extremes or original number based on masks.
+			VInOut = VectorSelect(MinMask, VMinVal, VInOut);
+			VInOut = VectorSelect(MaxMask, VMaxVal, VInOut);
+
+			VectorStoreAligned(VInOut, &InOutBuffer[i]);
+		}
+	}
+
 	void BufferMultiplyByConstant(const AlignedFloatBuffer& InFloatBuffer, float InValue, AlignedFloatBuffer& OutFloatBuffer)
 	{
 		check(InFloatBuffer.Num() >= 4);
diff --git a/Engine/Source/Runtime/SignalProcessing/Public/DSP/BufferVectorOperations.h b/Engine/Source/Runtime/SignalProcessing/Public/DSP/BufferVectorOperations.h
--- a/Engine/Source/Runtime/SignalProcessing/Public/DSP/BufferVectorOperations.h
+++ b/Engine/Source/Runtime/SignalProcessing/Public/DSP/BufferVectorOperations.h
@@ -31,7 +31,11 @@ namespace Audio
 
 	/* Sets a values to zero if value is denormal. Denormal numbers significantly slow down floating point operations. */
 	SIGNALPROCESSING_API void BufferUnderflowClampFast(float* RESTRICT InOutBuffer, const int32 InNum);
-		
+	
+	/* Clamps the values in a buffer between a min and max value. Returns true if it had to clamp */
+	void SIGNALPROCESSING_API BufferRangeClampFast(AlignedFloatBuffer& InOutBuffer, float InMinValue, float InMaxValue);
+	void SIGNALPROCESSING_API BufferRangeClampFast(float* RESTRICT InOutBuffer, const int32 InNum, float InMinValue, float InMaxValue);
+	
 	/** Multiplies the input aligned float buffer with the given value. */
 	SIGNALPROCESSING_API void BufferMultiplyByConstant(const AlignedFloatBuffer& InFloatBuffer, float InValue, AlignedFloatBuffer& OutFloatBuffer);
 	SIGNALPROCESSING_API void BufferMultiplyByConstant(const float* RESTRICT InFloatBuffer, float InValue, float* RESTRICT OutFloatBuffer, const int32 InNumSamples);
