Edigrate clamp on final audio output to avoid overflow on float to int16 conversion

diff --git a/Engine/Source/Runtime/SignalProcessing/Public/DSP/BufferVectorOperations.h b/Engine/Source/Runtime/SignalProcessing/Public/DSP/BufferVectorOperations.h
--- a/Engine/Source/Runtime/SignalProcessing/Public/DSP/BufferVectorOperations.h
+++ b/Engine/Source/Runtime/SignalProcessing/Public/DSP/BufferVectorOperations.h
@@ -32,7 +32,7 @@ namespace Audio
 	/* Sets a values to zero if value is denormal. Denormal numbers significantly slow down floating point operations. */
 	SIGNALPROCESSING_API void BufferUnderflowClampFast(float* RESTRICT InOutBuffer, const int32 InNum);
 	
-	/* Clamps the values in a buffer between a min and max value. Returns true if it had to clamp */
+	/* Clamps the values in a buffer between a min and max value.*/
 	void SIGNALPROCESSING_API BufferRangeClampFast(AlignedFloatBuffer& InOutBuffer, float InMinValue, float InMaxValue);
 	void SIGNALPROCESSING_API BufferRangeClampFast(float* RESTRICT InOutBuffer, const int32 InNum, float InMinValue, float InMaxValue);
 	
