VOIP Noise Gate.

diff --git a/Engine/Source/Runtime/Engine/Private/VoiceConfig.cpp b/Engine/Source/Runtime/Engine/Private/VoiceConfig.cpp
--- a/Engine/Source/Runtime/Engine/Private/VoiceConfig.cpp
+++ b/Engine/Source/Runtime/Engine/Private/VoiceConfig.cpp
@@ -29,6 +29,30 @@ FAutoConsoleVariableRef CVarJitterBufferDelay(
 	TEXT("Value: Number of seconds of audio we buffer."),
 	ECVF_Default);
 
+static float MicNoiseGateThresholdCvar = 0.08;
+FAutoConsoleVariableRef CVarMicNoiseGateThreshold(
+	TEXT("voice.MicNoiseGateThreshold"),
+	JitterBufferDelayCvar,
+	TEXT("Our threshold, in linear amplitude, for  our noise gate on input. Similar to voice.SilenceDetectionThreshold, except that audio quieter than our noise gate threshold will still output silence.\n")
+	TEXT("Value: Number of seconds of audio we buffer."),
+	ECVF_Default);
+
+static float MicNoiseGateAttackTimeCvar = 0.05;
+FAutoConsoleVariableRef CVarMicNoiseGateAttackTime(
+	TEXT("voice.MicNoiseAttackTime"),
+	MicNoiseGateAttackTimeCvar,
+	TEXT("Sets the fade-in time for our noise gate.\n")
+	TEXT("Value: Number of seconds we fade in over."),
+	ECVF_Default);
+
+static float MicNoiseGateReleaseTimeCvar = 0.30;
+FAutoConsoleVariableRef CVarMicNoiseGateReleaseTime(
+	TEXT("voice.MicNoiseReleaseTime"),
+	MicNoiseGateReleaseTimeCvar,
+	TEXT("Sets the fade out time for our noise gate.\n")
+	TEXT("Value: Number of seconds we fade out over."),
+	ECVF_Default);
+
 static int32 NumVoiceChannelsCvar = 1;
 FAutoConsoleVariableRef CVarNumVoiceChannels(
 	TEXT("voice.NumChannels"),
@@ -153,6 +177,11 @@ float UVOIPStatics::GetBufferingDelay()
 	return JitterBufferDelayCvar;
 }
 
+float UVOIPStatics::GetVoiceNoiseGateLevel()
+{
+	return MicNoiseGateThresholdCvar;
+}
+
 int32 UVOIPStatics::GetNumBufferedPackets()
 {
 	// Evaluate number of packets as the total number of samples we'll need to buffer divided by
diff --git a/Engine/Source/Runtime/Engine/Public/Net/VoiceConfig.h b/Engine/Source/Runtime/Engine/Public/Net/VoiceConfig.h
--- a/Engine/Source/Runtime/Engine/Public/Net/VoiceConfig.h
+++ b/Engine/Source/Runtime/Engine/Public/Net/VoiceConfig.h
@@ -164,6 +164,9 @@ public:
 	// Returns the amount of seconds to spend buffering specified by AudioSettings.
 	static float GetBufferingDelay();
 
+	// Returns the noise gate threshold we use for incoming audio.
+	static float GetVoiceNoiseGateLevel();
+
 	// returns the amount of packets to allocate memory for based on the buffering delay.
 	static int32 GetNumBufferedPackets();
 
diff --git a/Engine/Source/Runtime/Online/Voice/Private/Windows/VoiceCaptureWindows.cpp b/Engine/Source/Runtime/Online/Voice/Private/Windows/VoiceCaptureWindows.cpp
--- a/Engine/Source/Runtime/Online/Voice/Private/Windows/VoiceCaptureWindows.cpp
+++ b/Engine/Source/Runtime/Online/Voice/Private/Windows/VoiceCaptureWindows.cpp
@@ -3,6 +3,7 @@
 #include "VoiceCaptureWindows.h"
 #include "VoicePrivate.h"
 #include "VoiceModule.h"
+#include "DSP/Dsp.h"
 
 #if PLATFORM_SUPPORTS_VOICE_CAPTURE
 
@@ -120,6 +121,8 @@ bool FVoiceCaptureWindows::Init(const FString& DeviceName, int32 SampleRate, int
 	LookaheadBuffer.Init(AttackInSamples + 1);
 	LookaheadBuffer.SetDelay(AttackInSamples);
 
+	NoiseGateAttenuator.Init(SampleRate);
+
 	bIsMicActive = false;
 
 	return CreateCaptureBuffer(DeviceName.IsEmpty() ? VoiceDev->DefaultVoiceCaptureDevice.DeviceName : DeviceName, SampleRate, NumChannels);
@@ -419,7 +422,19 @@ void FVoiceCaptureWindows::ProcessData()
 			static IConsoleVariable* SilenceDetectionThresholdCVar = IConsoleManager::Get().FindConsoleVariable(TEXT("voice.SilenceDetectionThreshold"));
 			check(SilenceDetectionThresholdCVar);			
 			const float MicSilenceThreshold = SilenceDetectionThresholdCVar->GetFloat();
+
+			static IConsoleVariable* NoiseGateThresholdCVar = IConsoleManager::Get().FindConsoleVariable(TEXT("voice.MicNoiseGateThreshold"));
+			check(NoiseGateThresholdCVar);
+			const float MicNoiseGateThreshold = NoiseGateThresholdCVar->GetFloat();
 			
+			static IConsoleVariable* NoiseGateAttackTimeCVar = IConsoleManager::Get().FindConsoleVariable(TEXT("voice.MicNoiseAttackTime"));
+			check(NoiseGateAttackTimeCVar);
+			const float MicNoiseGateAttackTime = NoiseGateAttackTimeCVar->GetFloat();
+
+			static IConsoleVariable* NoiseGateReleaseTimeCVar = IConsoleManager::Get().FindConsoleVariable(TEXT("voice.MicNoiseReleaseTime"));
+			check(NoiseGateReleaseTimeCVar);
+			const float MicNoiseGateReleaseTime = NoiseGateReleaseTimeCVar->GetFloat();
+
 			bool bMicReleased = false;
 
 			CurrentSampleStart = CachedSampleStart;
@@ -440,12 +455,36 @@ void FVoiceCaptureWindows::ProcessData()
 
 				bIsMicActive = MicLevelDetector.GetCurrentValue() > MicSilenceThreshold;
 
+				// If we have just crossed the noise gate threshold, begin interpoloating to 1.0 or 0.0
+				const bool bIsMicAboveNoiseGateThreshold = MicLevelDetector.GetCurrentValue() > MicNoiseGateThreshold;
+
+				if (bIsMicAboveNoiseGateThreshold && !bWasMicAboveNoiseGateThreshold)
+				{
+					NoiseGateAttenuator.SetValue(1.0f, MicNoiseGateAttackTime);
+				}
+				else if (!bIsMicAboveNoiseGateThreshold && bWasMicAboveNoiseGateThreshold)
+				{
+					NoiseGateAttenuator.SetValue(0.0f, MicNoiseGateReleaseTime);
+				}
+
+				bWasMicAboveNoiseGateThreshold = bIsMicAboveNoiseGateThreshold;
+
 				if (bIsMicActive)
 				{
 					if (bMicReleased)
 					{
+						// Apply noise gate attenuation.
+						for (int32 ChannelIndex = 0; ChannelIndex < NumInputChannels; ChannelIndex++)
+						{
+							Audio::TSampleRef<int16> SampleRef(InputBuffer[FrameIndex + ChannelIndex]);
+							SampleRef = SampleRef * NoiseGateAttenuator.PeekCurrentValue();
+						}
+
+						NoiseGateAttenuator.GetNextValue();
+
 						ReleaseBuffer.PushFrame(&InputBuffer[FrameIndex], NumInputChannels);
 
+
 						if (!bSampleStartCached)
 						{
 							CachedSampleStart = SampleCounter;
@@ -454,6 +493,15 @@ void FVoiceCaptureWindows::ProcessData()
 					}
 					else
 					{
+						// Apply noise gate attenuation.
+						for (int32 ChannelIndex = 0; ChannelIndex < NumInputChannels; ChannelIndex++)
+						{
+							Audio::TSampleRef<int16> SampleRef(InputBuffer[FrameIndex + ChannelIndex]);
+							SampleRef = SampleRef * NoiseGateAttenuator.PeekCurrentValue();
+						}
+
+						NoiseGateAttenuator.GetNextValue();
+
 						FMemory::Memcpy(&AudioBuffer[FrameIndex], &InputBuffer[FrameIndex], sizeof(int16) * NumInputChannels);
 						SamplesPushedToUncompressedAudioBuffer += NumInputChannels;
 					}
@@ -484,10 +532,33 @@ void FVoiceCaptureWindows::ProcessData()
 
 				bIsMicActive = MicLevelDetector.GetCurrentValue() > MicSilenceThreshold;
 
+				// If we have just crossed the noise gate threshold, begin interpoloating to 1.0 or 0.0
+				const bool bIsMicAboveNoiseGateThreshold = MicLevelDetector.GetCurrentValue() > MicNoiseGateThreshold;
+
+				if (bIsMicAboveNoiseGateThreshold && !bWasMicAboveNoiseGateThreshold)
+				{
+					NoiseGateAttenuator.SetValue(1.0f, MicNoiseGateAttackTime);
+				}
+				else if (!bIsMicAboveNoiseGateThreshold && bWasMicAboveNoiseGateThreshold)
+				{
+					NoiseGateAttenuator.SetValue(0.0f, MicNoiseGateReleaseTime);
+				}
+
+				bWasMicAboveNoiseGateThreshold = bIsMicAboveNoiseGateThreshold;
+
 				if (bIsMicActive)
 				{
 					if (bMicReleased)
 					{
+						// Apply noise gate attenuation.
+						for (int32 ChannelIndex = 0; ChannelIndex < NumInputChannels; ChannelIndex++)
+						{
+							Audio::TSampleRef<int16> SampleRef(InputBuffer[FrameIndex + ChannelIndex]);
+							SampleRef = SampleRef * NoiseGateAttenuator.PeekCurrentValue();
+						}
+
+						NoiseGateAttenuator.GetNextValue();
+
 						ReleaseBuffer.PushFrame(&InputBuffer[FrameIndex], NumInputChannels);
 
 						if (!bSampleStartCached)
@@ -498,6 +569,15 @@ void FVoiceCaptureWindows::ProcessData()
 					}
 					else
 					{
+						// Apply noise gate attenuation.
+						for (int32 ChannelIndex = 0; ChannelIndex < NumInputChannels; ChannelIndex++)
+						{
+							Audio::TSampleRef<int16> SampleRef(InputBuffer[FrameIndex + ChannelIndex]);
+							SampleRef = SampleRef * NoiseGateAttenuator.PeekCurrentValue();
+						}
+
+						NoiseGateAttenuator.GetNextValue();
+
 						FMemory::Memcpy(&AudioBuffer[FrameIndex], &InputBuffer[FrameIndex], sizeof(int16) * NumInputChannels);
 						SamplesPushedToUncompressedAudioBuffer += NumInputChannels;
 					}
diff --git a/Engine/Source/Runtime/Online/Voice/Private/Windows/VoiceCaptureWindows.h b/Engine/Source/Runtime/Online/Voice/Private/Windows/VoiceCaptureWindows.h
--- a/Engine/Source/Runtime/Online/Voice/Private/Windows/VoiceCaptureWindows.h
+++ b/Engine/Source/Runtime/Online/Voice/Private/Windows/VoiceCaptureWindows.h
@@ -173,6 +173,16 @@ private:
 	*/
 	bool bIsMicActive;
 
+	/**
+	 * This linear ease is used to add attack fade in and release fadeout to the noise gate.
+	 */
+	Audio::FLinearEase NoiseGateAttenuator;
+
+	/**
+	 * This bool is used to track when we engage 
+	 */
+	bool bWasMicAboveNoiseGateThreshold;
+
 	/**
 	 * Create the D3D8 capture device
 	 *
diff --git a/Engine/Source/Runtime/SignalProcessing/Public/DSP/Dsp.h b/Engine/Source/Runtime/SignalProcessing/Public/DSP/Dsp.h
--- a/Engine/Source/Runtime/SignalProcessing/Public/DSP/Dsp.h
+++ b/Engine/Source/Runtime/SignalProcessing/Public/DSP/Dsp.h
@@ -4,6 +4,8 @@
 
 #include "CoreMinimal.h"
 #include "Misc/ScopeLock.h"
+#include "Templates/IsFloatingPoint.h"
+#include "Templates/IsIntegral.h"
 
 // Macros which can be enabled to cause DSP sample checking
 #if 0
@@ -510,7 +512,7 @@ namespace Audio
 	 * Designed to be thread safe for SPSC; However, if Push() and Pop() are both trying to access an overlapping area of the buffer,
 	 * One of the calls will be truncated. Thus, it is advised that you use a high enough capacity that the producer and consumer are never in contention.
 	 */
-	template <class SampleType>
+	template <typename SampleType>
 	class TCircularAudioBuffer
 	{
 	private:
@@ -645,5 +647,378 @@ namespace Audio
 			return (Capacity - 1 - WriteIndex + ReadIndex) % Capacity;
 		}
 	};
+
+	/**
+	 * This allows us to write a compile time exponent of a number.
+	 */
+	template <int Base, int Exp>
+	struct TGetPower
+	{
+		static_assert(Exp >= 0, "TGetPower only supports positive exponents.");
+		static const int64 Value = Base * TGetPower<Base, Exp - 1>::Value;
+	};
+
+	template <int Base>
+	struct TGetPower<Base, 0>
+	{
+		static const int64 Value = 1;
+	};
+
+	/**
+	 * TSample<SampleType, Q>
+	 * Variant type to simplify converting and performing operations on fixed precision and floating point samples.
+	 */
+	template <typename SampleType, uint32 Q = (sizeof(SampleType) * 8 - 1)>
+	class TSample
+	{
+		SampleType Sample;
+
+		template <typename SampleTypeToCheck>
+		static void CheckValidityOfSampleType()
+		{
+			constexpr bool bIsTypeValid = !!(TIsFloatingPoint<SampleTypeToCheck>::Value || TIsIntegral<SampleTypeToCheck>::Value);
+			static_assert(bIsTypeValid, "Invalid sample type! TSampleRef only supports float or integer values.");
+		}
+
+		template <typename SampleTypeToCheck, uint32 QToCheck>
+		static void CheckValidityOfQ()
+		{
+			// If this is a fixed-precision value, our Q offset must be less than how many bits we have.
+			constexpr bool bIsTypeValid = !!(TIsFloatingPoint<SampleTypeToCheck>::Value || ((sizeof(SampleTypeToCheck) * 8) > QToCheck));
+			static_assert(bIsTypeValid, "Invalid value for Q! TSampleRef only supports float or int types. For int types, Q must be smaller than the number of bits in the int type.");
+		}
+
+		// This is the number used to convert from float to our fixed precision value.
+		static constexpr float QFactor = TGetPower<2, Q>::Value - 1;
+
+	public:
+
+		TSample(SampleType& InSample)
+			: Sample(InSample)
+		{
+			CheckValidityOfQ<SampleType, Q>();
+			CheckValidityOfSampleType<SampleType>();
+		}
+
+		template<typename ReturnType = float>
+		ReturnType AsFloat() const
+		{
+			static_assert(TIsFloatingPoint<ReturnType>::Value, "Return type for AsFloat() must be a floating point type.");
+
+			if (TIsFloatingPoint<SampleType>::Value)
+			{
+				return static_cast<ReturnType>(Sample);
+			}
+			else if (TIsIntegral<SampleType>::Value)
+			{
+				// Cast from fixed to float.
+				return static_cast<ReturnType>(Sample) / QFactor;
+			}
+			else
+			{
+				checkNoEntry();
+				return static_cast<ReturnType>(Sample);
+			}
+		}
+
+		template<typename ReturnType, uint32 ReturnQ = (sizeof(SampleType) * 8 - 1)>
+		ReturnType AsFixedPrecisionInt()
+		{
+			static_assert(TIsIntegral<ReturnType>::Value, "This function must be called with an integer type as ReturnType.");
+			CheckValidityOfQ<ReturnType, ReturnQ>();
+
+			if (TIsIntegral<SampleType>::Value)
+			{
+				if (Q > ReturnQ)
+				{
+					return Sample << (Q - ReturnQ);
+				}
+				else if (Q < ReturnQ)
+				{
+					return Sample >> (ReturnQ - Q);
+				}
+				else
+				{
+					return Sample;
+				}
+			}
+			else if (TIsFloatingPoint<SampleType>::Value)
+			{
+				static constexpr float ReturnQFactor = TGetPower<2, ReturnQ>::Value - 1;
+				return (ReturnType)(Sample * ReturnQFactor);
+			}
+		}
+
+		template <typename OtherSampleType>
+		TSample<SampleType, Q>& operator =(const OtherSampleType InSample)
+		{
+			CheckValidityOfSampleType<OtherSampleType>();
+
+			if (TIsSame<SampleType, OtherSampleType>::Value)
+			{
+				Sample = InSample;
+				return *this;
+			}
+			else if (TIsIntegral<OtherSampleType>::Value && TIsFloatingPoint<SampleType>::Value)
+			{
+				// Cast from Q15 to float.
+				Sample = ((SampleType)InSample) / QFactor;
+				return *this;
+			}
+			else if (TIsFloatingPoint<OtherSampleType>::Value && TIsIntegral<SampleType>::Value)
+			{
+				// cast from float to Q15.
+				Sample = static_cast<SampleType>(InSample * QFactor);
+				return *this;
+			}
+			else
+			{
+				checkNoEntry();
+				return *this;
+			}
+		}
+
+		template <typename OtherSampleType>
+		friend TSample<SampleType, Q> operator *(const TSample<SampleType, Q>& LHS, const OtherSampleType& RHS)
+		{
+			CheckValidityOfSampleType<OtherSampleType>();
+
+			// Float case:
+			if (TIsFloatingPoint<SampleType>::Value)
+			{
+				if (TIsFloatingPoint<OtherSampleType>::Value)
+				{
+					// float * float.
+					return LHS.Sample * RHS;
+				}
+				else if (TIsIntegral<OtherSampleType>::Value)
+				{
+					// float * Q.
+					SampleType FloatRHS = ((SampleType)RHS) / QFactor;
+					return LHS.Sample * FloatRHS;
+				}
+				else
+				{
+					checkNoEntry();
+					return LHS.Sample;
+				}
+
+			}
+			// Q Case
+			else if (TIsIntegral<SampleType>::Value)
+			{
+				if (TIsFloatingPoint<OtherSampleType>::Value)
+				{
+					// fixed * float.
+					OtherSampleType FloatLHS = ((OtherSampleType)LHS.Sample) / QFactor;
+					OtherSampleType Result = FloatLHS * RHS;
+					return static_cast<SampleType>(Result * QFactor);
+				}
+				else if (TIsIntegral<OtherSampleType>::Value)
+				{
+					// Q * Q.
+					float FloatLHS = ((float)LHS.Sample) / QFactor;
+					float FloatRHS = ((float)RHS) / QFactor;
+					float Result = FloatLHS * FloatRHS;
+					return static_cast<OtherSampleType>(Result * QFactor);
+				}
+				else
+				{
+					checkNoEntry();
+					return LHS.Sample;
+				}
+			}
+			else
+			{
+				checkNoEntry();
+				return LHS.Sample;
+			}
+		}
+	};
+
+
+	/**
+	 * TSampleRef<SampleType, Q>
+	 * Ref version of TSample. Useful for converting between fixed and float precisions.
+	 * Example usage:
+	 * int16 FixedPrecisionSample;
+	 * TSampleRef<int16, 15> SampleRef(FixedPrecisionSample);
+	 * 
+	 * // Set the sample value directly:
+	 * SampleRef = 0.5f;
+	 * 
+	 * // Or multiply the the sample:
+	 * SampleRef *= 0.5f;
+	 *
+	 * bool bThisCodeWorks = FixedPrecisionSample == TNumericLimits<int16>::Max() / 4;
+	 */
+	template <typename SampleType, uint32 Q = (sizeof(SampleType) * 8 - 1)>
+	class TSampleRef
+	{
+		SampleType& Sample;
+
+		template <typename SampleTypeToCheck>
+		static void CheckValidityOfSampleType()
+		{
+			constexpr bool bIsTypeValid = !!(TIsFloatingPoint<SampleTypeToCheck>::Value || TIsIntegral<SampleTypeToCheck>::Value);
+			static_assert(bIsTypeValid, "Invalid sample type! TSampleRef only supports float or integer values.");
+		}
+
+		template <typename SampleTypeToCheck, uint32 QToCheck>
+		static void CheckValidityOfQ()
+		{
+			// If this is a fixed-precision value, our Q offset must be less than how many bits we have.
+			constexpr bool bIsTypeValid = !!(TIsFloatingPoint<SampleTypeToCheck>::Value || (sizeof(SampleTypeToCheck) * 8) > QToCheck);
+			static_assert(bIsTypeValid, "Invalid value for Q! TSampleRef only supports float or int types. For int types, Q must be smaller than the number of bits in the int type.");
+		}
+
+		// This is the number used to convert from float to our .
+		static constexpr float QFactor = TGetPower<2, Q>::Value - 1;
+
+	public:
+
+		TSampleRef(SampleType& InSample)
+			: Sample(InSample)
+		{
+			CheckValidityOfQ<SampleType, Q>();
+			CheckValidityOfSampleType<SampleType>();
+		}
+
+		template<typename ReturnType = float>
+		ReturnType AsFloat() const
+		{
+			static_assert(TIsFloatingPoint<ReturnType>::Value, "Return type for AsFloat() must be a floating point type.");
+
+			if (TIsFloatingPoint<SampleType>::Value)
+			{
+				return static_cast<ReturnType>(Sample);
+			}
+			else if (TIsIntegral<SampleType>::Value)
+			{
+				// Cast from fixed to float.
+				return static_cast<ReturnType>(Sample) / QFactor;
+			}
+			else
+			{
+				checkNoEntry();
+				return static_cast<ReturnType>(Sample);
+			}
+		}
+
+		template<typename ReturnType, uint32 ReturnQ = (sizeof(SampleType) * 8 - 1)>
+		ReturnType AsFixedPrecisionInt()
+		{
+			static_assert(TIsIntegral<ReturnType>::Value, "This function must be called with an integer type as ReturnType.");
+			
+			CheckValidityOfQ<ReturnType, ReturnQ>();
+
+			if (TIsIntegral<SampleType>::Value)
+			{
+				if (Q > ReturnQ)
+				{
+					return Sample << (Q - ReturnQ);
+				}
+				else if (Q < ReturnQ)
+				{
+					return Sample >> (ReturnQ - Q);
+				}
+				else
+				{
+					return Sample;
+				}
+			}
+			else if (TIsFloatingPoint<SampleType>::Value)
+			{
+				static constexpr SampleType ReturnQFactor = TGetPower<2, ReturnQ>::Value - 1;
+				return (ReturnType) (Sample * ReturnQFactor);
+			}
+		}
+
+		template <typename OtherSampleType>
+		TSampleRef<SampleType, Q>& operator =(const OtherSampleType InSample)
+		{
+			CheckValidityOfSampleType<OtherSampleType>();
+
+			if (TIsSame<SampleType, OtherSampleType>::Value)
+			{
+				Sample = InSample;
+				return *this;
+			}
+			else if (TIsIntegral<OtherSampleType>::Value && TIsFloatingPoint<SampleType>::Value)
+			{
+				// Cast from fixed to float.
+				Sample = ((SampleType)InSample) / QFactor;
+				return *this;
+			}
+			else if (TIsFloatingPoint<OtherSampleType>::Value && TIsIntegral<SampleType>::Value)
+			{
+				// cast from float to fixed.
+				Sample = (SampleType)(InSample * QFactor);
+				return *this;
+			}
+			else
+			{
+				checkNoEntry();
+				return *this;
+			}
+		}
+
+		template <typename OtherSampleType>
+		friend SampleType operator *(const TSampleRef<SampleType>& LHS, const OtherSampleType& RHS)
+		{
+			CheckValidityOfSampleType<OtherSampleType>();
+
+			// Float case:
+			if (TIsFloatingPoint<SampleType>::Value)
+			{
+				if (TIsFloatingPoint<OtherSampleType>::Value)
+				{
+					// float * float.
+					return LHS.Sample * RHS;
+				}
+				else if (TIsIntegral<OtherSampleType>::Value)
+				{
+					// float * fixed.
+					SampleType FloatRHS = ((SampleType)RHS) / QFactor;
+					return LHS.Sample * FloatRHS;
+				}
+				else
+				{
+					checkNoEntry();
+					return LHS.Sample;
+				}
+
+			}
+			// Fixed Precision Case
+			else if (TIsIntegral<SampleType>::Value)
+			{
+				if (TIsFloatingPoint<OtherSampleType>::Value)
+				{
+					// fixed * float.
+					OtherSampleType FloatLHS = ((OtherSampleType)LHS.Sample) / QFactor;
+					OtherSampleType Result = FloatLHS * RHS;
+					return static_cast<SampleType>(Result * QFactor);
+				}
+				else if (TIsIntegral<OtherSampleType>::Value)
+				{
+					// Q * Q.
+					float FloatLHS = ((float)LHS.Sample) / QFactor;
+					float FloatRHS = ((float)RHS) / QFactor;
+					float Result = FloatLHS * FloatRHS;
+					return static_cast<SampleType>(Result * QFactor);
+				}
+				else
+				{
+					checkNoEntry();
+					return LHS.Sample;
+				}
+			}
+			else
+			{
+				checkNoEntry();
+				return LHS.Sample;
+			}
+		}
+	};
 }
 
