Fix clang compile error in unused template

diff --git a/Engine/Source/Runtime/SignalProcessing/Public/DSP/Dsp.h b/Engine/Source/Runtime/SignalProcessing/Public/DSP/Dsp.h
--- a/Engine/Source/Runtime/SignalProcessing/Public/DSP/Dsp.h
+++ b/Engine/Source/Runtime/SignalProcessing/Public/DSP/Dsp.h
@@ -693,8 +693,8 @@ namespace Audio
 		static constexpr float QFactor = TGetPower<2, Q>::Value - 1;
 
 		// for fixed precision types, the max and min values that we can represent are calculated here:
-		static constexpr float MaxValue = TGetPower<2, (sizeof(SampleType) * 8 - Q)>;
-		static constexpr float MinValue = !!TIsSigned<SampleType> ? (-1.0f * MaxValue) : 0.0f;
+		static constexpr float MaxValue = TGetPower<2, (sizeof(SampleType) * 8 - Q)>::Value;
+		static constexpr float MinValue = !!TIsSigned<SampleType>::Value ? (-1.0f * MaxValue) : 0.0f;
 
 	public:
 
