- Fix bypass implementation in Submix Effect ReverbFast
- Minor clean-up of reverb fast private DSP API
- Clamp Dry/Wet on reverb UX to values clamped in DSP API

diff --git a/Engine/Source/Runtime/AudioMixer/Classes/SubmixEffects/AudioMixerSubmixEffectReverbFast.h b/Engine/Source/Runtime/AudioMixer/Classes/SubmixEffects/AudioMixerSubmixEffectReverbFast.h
--- a/Engine/Source/Runtime/AudioMixer/Classes/SubmixEffects/AudioMixerSubmixEffectReverbFast.h
+++ b/Engine/Source/Runtime/AudioMixer/Classes/SubmixEffects/AudioMixerSubmixEffectReverbFast.h
@@ -2,9 +2,11 @@
 
 #pragma once
 
+#include "AudioEffect.h"
 #include "Sound/SoundEffectSubmix.h"
+#include "DSP/Amp.h"
 #include "DSP/ReverbFast.h"
-#include "AudioEffect.h"
+
 #include "AudioMixerSubmixEffectReverbFast.generated.h"
 
 USTRUCT(BlueprintType)
@@ -60,12 +62,12 @@ struct AUDIOMIXER_API FSubmixEffectReverbFastSettings
 	UPROPERTY(EditAnywhere, BlueprintReadWrite, Category = General, meta = (ClampMin = "0.0", ClampMax = "1.0", EditCondition = "!bBypass"))
 	float AirAbsorptionGainHF;
 
-	// Overall wetlevel of the reverb effect
-	UPROPERTY(EditAnywhere, BlueprintReadWrite, Category = Routing, meta = (EditCondition = "!bBypass"))
+	// Overall wet level of the reverb effect
+	UPROPERTY(EditAnywhere, BlueprintReadWrite, Category = Routing, meta = (EditCondition = "!bBypass", UIMin = "0.0", UIMax = "1.0", ClampMin = "0.0", ClampMax = "10.0", EditCondition = "!bBypass"))
 	float WetLevel;
 
-	// Overall drylevel of the reverb effect
-	UPROPERTY(EditAnywhere, BlueprintReadWrite, Category = Routing, meta = (EditCondition = "!bBypass"))
+	// Overall dry level of the reverb effect
+	UPROPERTY(EditAnywhere, BlueprintReadWrite, Category = Routing, meta = (EditCondition = "!bBypass", ClampMin = "0.0", ClampMax = "1.0", EditCondition = "!bBypass"))
 	float DryLevel;
 
 	FSubmixEffectReverbFastSettings()
@@ -108,18 +110,26 @@ public:
 	void SetEffectParameters(const FAudioReverbEffect& InReverbEffectParameters);
 
 private:
+
+	static const float MinWetness;
+	static const float MaxWetness;
+
 	void UpdateParameters();
 
 	// The fast reverb effect
 	TUniquePtr<Audio::FPlateReverbFast> PlateReverb;
 
 	// The reverb effect params
-	Audio::TParams<Audio::FPlateReverbFastSettings> Params;
+	Audio::TParams<Audio::FPlateReverbFastSettings> ReverbParams;
+
+	// Settings for wet and dry signal to be consumed on next buffer
+	Audio::TParams<Audio::FWetDry> WetDryParams;
+
+	// Level of wet/dry signal on current buffer
+	Audio::FWetDry CurrentWetDry;
 
 	// Curve which maps old reverb times to new decay value
 	FRichCurve DecayCurve;
-
-	bool bBypass;
 };
 
 UCLASS()
diff --git a/Engine/Source/Runtime/AudioMixer/Private/Effects/AudioMixerSubmixEffectReverbFast.cpp b/Engine/Source/Runtime/AudioMixer/Private/Effects/AudioMixerSubmixEffectReverbFast.cpp
--- a/Engine/Source/Runtime/AudioMixer/Private/Effects/AudioMixerSubmixEffectReverbFast.cpp
+++ b/Engine/Source/Runtime/AudioMixer/Private/Effects/AudioMixerSubmixEffectReverbFast.cpp
@@ -5,6 +5,7 @@
 #include "Sound/ReverbEffect.h"
 #include "Audio.h"
 #include "AudioMixer.h"
+#include "DSP/BufferVectorOperations.h"
 #include "DSP/ReverbFast.h"
 
 
@@ -34,11 +35,11 @@ static FAutoConsoleVariableRef CVarDisableQuadReverbCVarFast(
 	ECVF_Default);
 
 
-
-class UReverbEffect;
+const float FSubmixEffectReverbFast::MinWetness = 0.0f;
+const float FSubmixEffectReverbFast::MaxWetness = 10.f;
 
 FSubmixEffectReverbFast::FSubmixEffectReverbFast()
-	: bBypass(false)
+	: CurrentWetDry(-1.0f, -1.0f)
 {
 }
 
@@ -57,9 +58,6 @@ void FSubmixEffectReverbFast::Init(const FSoundEffectSubmixInitData& InitData)
 	 */
 	Audio::FPlateReverbFastSettings NewSettings;
 
-
-	NewSettings.Wetness = 1.0f;
-
 	NewSettings.EarlyReflections.Decay = 0.9f;
 	NewSettings.EarlyReflections.Absorption = 0.7f;
 	NewSettings.EarlyReflections.Gain = 1.0f;
@@ -74,7 +72,7 @@ void FSubmixEffectReverbFast::Init(const FSoundEffectSubmixInitData& InitData)
 	NewSettings.LateReflections.Decay = 0.15f;
 	NewSettings.LateReflections.Density = 0.85f;
 
-	Params.SetParams(NewSettings);
+	ReverbParams.SetParams(NewSettings);
 
 	DecayCurve.AddKey(0.0f, 0.99f);
 	DecayCurve.AddKey(2.0f, 0.45f);
@@ -108,11 +106,15 @@ void FSubmixEffectReverbFast::OnPresetChanged()
 	ReverbEffect.LateDelay = Settings.LateDelay;
 	ReverbEffect.AirAbsorptionGainHF = Settings.AirAbsorptionGainHF;
 	ReverbEffect.RoomRolloffFactor = 0.0f; // not used
-	ReverbEffect.Volume = Settings.WetLevel;
-
-	bBypass = Settings.bBypass;
+	ReverbEffect.Volume = Settings.bBypass ? 0.0f : Settings.WetLevel;
 
 	SetEffectParameters(ReverbEffect);
+
+	Audio::FWetDry NewWetDryParams;
+	NewWetDryParams.DryLevel = Settings.bBypass ? 1.0f : Settings.DryLevel;
+	NewWetDryParams.WetLevel = Settings.bBypass ? 0.0f : FMath::Clamp(Settings.WetLevel, MinWetness, MaxWetness);
+
+	WetDryParams.SetParams(NewWetDryParams);
 }
 
 void FSubmixEffectReverbFast::OnProcessAudio(const FSoundEffectSubmixInputData& InData, FSoundEffectSubmixOutputData& OutData)
@@ -126,14 +128,25 @@ void FSubmixEffectReverbFast::OnProcessAudio(const FSoundEffectSubmixInputData&
 		return;
 	}
 
-	if (bBypass)
+	UpdateParameters();
+
+	float LastWet = CurrentWetDry.WetLevel;
+	float LastDry = CurrentWetDry.DryLevel;
+	WetDryParams.GetParams(&CurrentWetDry);
+
+	// Set to most recent if uninitialized (less than 0.0f)
+	if (LastWet < 0.0f)
 	{
-		return;
+		LastWet = CurrentWetDry.WetLevel;
+	}
+	if (LastDry < 0.0f)
+	{
+		LastDry = CurrentWetDry.DryLevel;
 	}
-
-	UpdateParameters();
 
 	PlateReverb->ProcessAudio(*InData.AudioBuffer, InData.NumChannels, *OutData.AudioBuffer, OutData.NumChannels);
+	Audio::FadeBufferFast(*OutData.AudioBuffer, LastWet, CurrentWetDry.WetLevel);
+	Audio::MixInBufferFast(*InData.AudioBuffer, *OutData.AudioBuffer, LastDry, CurrentWetDry.DryLevel);
 }
 
 void FSubmixEffectReverbFast::SetEffectParameters(const FAudioReverbEffect& InParams)
@@ -151,9 +164,6 @@ void FSubmixEffectReverbFast::SetEffectParameters(const FAudioReverbEffect& InPa
 	 */
 	Audio::FPlateReverbFastSettings NewSettings;
 
-	// General reverb settings
-	NewSettings.Wetness = FMath::GetMappedRangeValueClamped({ 0.0f, 10.0f }, { 0.0f, 10.0f }, InParams.Volume);
-
 	// Early Reflections
 	NewSettings.EarlyReflections.Gain = FMath::GetMappedRangeValueClamped({ 0.0f, 3.16f }, { 0.0f, 1.0f }, InParams.ReflectionsGain);
 	NewSettings.EarlyReflections.PreDelayMsec = FMath::GetMappedRangeValueClamped({ 0.0f, 0.3f }, { 0.0f, 300.0f }, InParams.ReflectionsDelay);
@@ -175,13 +185,13 @@ void FSubmixEffectReverbFast::SetEffectParameters(const FAudioReverbEffect& InPa
 	NewSettings.LateReflections.LateGainDB = Audio::ConvertToDecibels(NewSettings.LateReflections.LateGainDB);
 
 	// Apply the settings the thread safe settings object
-	Params.SetParams(NewSettings);
+	ReverbParams.SetParams(NewSettings);
 }
 
 void FSubmixEffectReverbFast::UpdateParameters()
 {
 	Audio::FPlateReverbFastSettings NewSettings;
-	if (Params.GetParams(&NewSettings))
+	if (ReverbParams.GetParams(&NewSettings))
 	{
 		PlateReverb->SetSettings(NewSettings);
 	}
@@ -215,7 +225,7 @@ void FSubmixEffectReverbFast::UpdateParameters()
 	}
 }
 
-void USubmixEffectReverbFastPreset::SetSettingsWithReverbEffect(const UReverbEffect* InReverbEffect, const float WetLevel, const float DryLevel)
+void USubmixEffectReverbFastPreset::SetSettingsWithReverbEffect(const UReverbEffect* InReverbEffect, const float InWetLevel, const float InDryLevel)
 {
 	if (InReverbEffect)
 	{
@@ -230,8 +240,8 @@ void USubmixEffectReverbFastPreset::SetSettingsWithReverbEffect(const UReverbEff
 		Settings.LateGain = InReverbEffect->LateGain;
 		Settings.LateDelay = InReverbEffect->LateDelay;
 		Settings.AirAbsorptionGainHF = InReverbEffect->AirAbsorptionGainHF;
-		Settings.WetLevel = WetLevel;
-		Settings.DryLevel = DryLevel;
+		Settings.WetLevel = InWetLevel;
+		Settings.DryLevel = InDryLevel;
 
 		Update();
 	}
diff --git a/Engine/Source/Runtime/SignalProcessing/Private/ReverbFast.cpp b/Engine/Source/Runtime/SignalProcessing/Private/ReverbFast.cpp
--- a/Engine/Source/Runtime/SignalProcessing/Private/ReverbFast.cpp
+++ b/Engine/Source/Runtime/SignalProcessing/Private/ReverbFast.cpp
@@ -6,8 +6,7 @@
 
 namespace Audio {
 	FPlateReverbFastSettings::FPlateReverbFastSettings()
-		: Wetness(0.5f)
-		, QuadBehavior(EQuadBehavior::StereoOnly)
+		: QuadBehavior(EQuadBehavior::StereoOnly)
 	{}
 
 	bool FPlateReverbFastSettings::operator==(const FPlateReverbFastSettings& Other) const
@@ -15,7 +14,6 @@ namespace Audio {
 		bool bIsEqual = (
 			(Other.EarlyReflections == EarlyReflections) &&
 			(Other.LateReflections == LateReflections) &&
-			(Other.Wetness == Wetness) &&
 			(Other.QuadBehavior == QuadBehavior));
 		
 		return bIsEqual;
@@ -26,16 +24,10 @@ namespace Audio {
 		return !(*this == Other);
 	}
 
-
-	const float FPlateReverbFast::MaxWetness = 10.f;
-	const float FPlateReverbFast::MinWetness = 0.0f;
-
 	const FPlateReverbFastSettings FPlateReverbFast::DefaultSettings;
 
 	FPlateReverbFast::FPlateReverbFast(float InSampleRate, int32 InMaxInternalBufferSamples, const FPlateReverbFastSettings& InSettings)
 		: SampleRate(InSampleRate)
-		, LastWetness(0.0f)
-		, bProcessCallSinceWetnessChanged(false)
 		, EarlyReflections(InSampleRate, InMaxInternalBufferSamples)
 		, LateReflections(InSampleRate, InMaxInternalBufferSamples, InSettings.LateReflections)
 		, bEnableEarlyReflections(true)
@@ -49,13 +41,6 @@ namespace Audio {
 
 	void FPlateReverbFast::SetSettings(const FPlateReverbFastSettings& InSettings)
 	{
-		// Copy, clamp and apply settings
-		if (bProcessCallSinceWetnessChanged)
-		{
-			LastWetness = Settings.Wetness;
-			bProcessCallSinceWetnessChanged = false;
-		}
-
 		Settings = InSettings;
 		ClampSettings(Settings);
 		ApplySettings();
@@ -93,18 +78,6 @@ namespace Audio {
 
 		FMemory::Memcpy(ScaledInputBuffer.GetData(), InSamples.GetData(), InSamples.Num() * sizeof(float));
 
-		// Scale input by wetness (or fade to new wetness)
-		if (FMath::IsNearlyEqual(LastWetness, Settings.Wetness))
-		{
-			MultiplyBufferByConstantInPlace(ScaledInputBuffer.GetData(), InSamples.Num(), Settings.Wetness);
-		}
-		else
-		{
-			FadeBufferFast(ScaledInputBuffer, LastWetness, Settings.Wetness);
-			LastWetness = Settings.Wetness;
-		}
-
-
 		checkf((1 == InNumChannels) || (2 == InNumChannels), TEXT("FPlateReverbFast only supports 1 or 2 channel inputs."))
 		checkf(OutNumChannels >= 2, TEXT("FPlateReverbFast requires at least 2 output channels."))
 
@@ -125,9 +98,7 @@ namespace Audio {
 			const int32 OutNum = InNumFrames * OutNumChannels;
 
 			OutSamples.Reset(OutNum);
-			OutSamples.AddUninitialized(OutNum);
-
-			FMemory::Memset(OutSamples.GetData(), 0, sizeof(float) * OutNum);
+			OutSamples.AddZeroed(OutNum);
 			return;
 		}
 
@@ -172,72 +143,15 @@ namespace Audio {
 
 		// Interleave and upmix
 		InterleaveAndMixOutput(FrontLeftReverbSamples, FrontRightReverbSamples, OutSamples, OutNumChannels);
-		bProcessCallSinceWetnessChanged = true;
 	}
 
 	void FPlateReverbFast::ClampSettings(FPlateReverbFastSettings& InOutSettings)
 	{
 		// Clamp settings for this object and member objects.
-		InOutSettings.Wetness = FMath::Clamp(InOutSettings.Wetness, MinWetness, MaxWetness);
 		FLateReflectionsFast::ClampSettings(InOutSettings.LateReflections);
 		FEarlyReflectionsFast::ClampSettings(InOutSettings.EarlyReflections);
 	}
 
-	// Copy input samples to output samples. Remap channels if necessary.
-	void FPlateReverbFast::PassThroughAudio(const AlignedFloatBuffer& InSamples, const int32 InNumChannels, AlignedFloatBuffer& OutSamples, const int32 OutNumChannels)
-	{
-		const int32 InNum = InSamples.Num();
-		const int32 InNumFrames = InNum / InNumChannels;
-		const int32 OutNum = OutNumChannels * InNumFrames;
-
-		// Resize output buffer
-		OutSamples.Reset(OutNum);
-		OutSamples.AddUninitialized(OutNum);
-		FMemory::Memset(OutSamples.GetData(), 0, sizeof(float) * OutNum);
-
-		if (InNum > 0)
-		{
-			if (InNumChannels == OutNumChannels)
-			{
-				FMemory::Memcpy(OutSamples.GetData(), InSamples.GetData(), sizeof(float) * InNum);
-			}
-			else
-			{
-				// InNumChannels can only be 1 or 2 channels so we have a limited number of 
-				// upmix situations.
-				FMemory::Memset(OutSamples.GetData(), 0, sizeof(float) * OutNum);
-				if (1 == InNumChannels)
-				{
-					// Upmix a mono signal
-					float* OutSampleData = OutSamples.GetData();
-					const float* InSampleData = InSamples.GetData();
-					int32 OutPos = 0;
-					for (int32 i = 0; i < InNumFrames; i++, OutPos += OutNumChannels)
-					{
-						// Scale to keep loudnes consistent.
-						float value = InSampleData[i] * 0.5f;	
-						OutSampleData[OutPos] = value;
-						OutSampleData[OutPos + 1] = value;
-					}
-				}
-				else if (2 == InNumChannels)
-				{
-					// Upmix a stereo signal by copying:
-					// 		FrontLeft  -> FrontLeft
-					//		FrontRight -> FrontRight
-					float* OutSampleData = OutSamples.GetData();
-					const float* InSampleData = InSamples.GetData();
-					int32 OutPos = 0;
-					for (int32 i = 0; i < InNum; i += InNumChannels, OutPos += OutNumChannels)
-					{
-						OutSampleData[OutPos] = InSampleData[i];
-						OutSampleData[OutPos + 1] = InSampleData[i + 1];
-					}
-				}
-			}
-		}
-	}
-
 	// Copy reverberated samples to interleaved output samples. Map channels according to internal settings.
 	// InFrontLeftSamples and InFrontRightSamples may be modified in-place.
 	void FPlateReverbFast::InterleaveAndMixOutput(const AlignedFloatBuffer& InFrontLeftSamples, const AlignedFloatBuffer& InFrontRightSamples, AlignedFloatBuffer& OutSamples, const int32 OutNumChannels)
diff --git a/Engine/Source/Runtime/SignalProcessing/Public/DSP/Amp.h b/Engine/Source/Runtime/SignalProcessing/Public/DSP/Amp.h
--- a/Engine/Source/Runtime/SignalProcessing/Public/DSP/Amp.h
+++ b/Engine/Source/Runtime/SignalProcessing/Public/DSP/Amp.h
@@ -7,6 +7,21 @@
 
 namespace Audio
 {
+	struct SIGNALPROCESSING_API FWetDry
+	{
+		// Level of wet signal
+		float WetLevel;
+
+		// Level of dry signal
+		float DryLevel;
+
+		FWetDry(float InWetLevel = 1.0f, float InDryLevel = 0.0f)
+			: WetLevel(InWetLevel)
+			, DryLevel(InDryLevel)
+		{
+		}
+	};
+
 	// Class which manages scaling audio input and performing panning operations
 	class SIGNALPROCESSING_API FAmp
 	{
diff --git a/Engine/Source/Runtime/SignalProcessing/Public/DSP/ReverbFast.h b/Engine/Source/Runtime/SignalProcessing/Public/DSP/ReverbFast.h
--- a/Engine/Source/Runtime/SignalProcessing/Public/DSP/ReverbFast.h
+++ b/Engine/Source/Runtime/SignalProcessing/Public/DSP/ReverbFast.h
@@ -40,9 +40,6 @@ namespace Audio
 		// from a sound source to the listener's ears. 
 		FLateReflectionsFastSettings LateReflections;
 
-		// Mix amount between dry and wet signals.
-		float Wetness;
-
 		// Set how reverb module generates reverb when there are 5 or more output channels. 
 		EQuadBehavior QuadBehavior;
 	};
@@ -52,8 +49,6 @@ namespace Audio
 	// class aims to support a flexible and pleasant sounding reverb balanced with computational efficiency. 
 	class SIGNALPROCESSING_API FPlateReverbFast {
 		public:
-			static const float MaxWetness;
-			static const float MinWetness;
 			static const FPlateReverbFastSettings DefaultSettings;
 			
 			// InMaxInternalBufferSamples sets the maximum number of samples used in internal buffers.
@@ -81,21 +76,13 @@ namespace Audio
 			static void ClampSettings(FPlateReverbFastSettings& InOutSettings);
 
 		private:
-			// Copy input samples to output samples. Remap channels if necessary.
-			void PassThroughAudio(const AlignedFloatBuffer& InSamples, const int32 InNumChannels, AlignedFloatBuffer& OutSamples, const int32 OutNumChannels);
-
-			
 			// Copy reverberated samples to interleaved output samples. Map channels according to internal settings.
 			void InterleaveAndMixOutput(const AlignedFloatBuffer& InFrontLeftSamples, const AlignedFloatBuffer& InFrontRightSamples, AlignedFloatBuffer& OutSamples, const int32 OutNumChannels);
 
 			void ApplySettings();
 
 
-
-
 			float SampleRate;
-			float LastWetness;
-			bool bProcessCallSinceWetnessChanged;
 
 			FPlateReverbFastSettings Settings;
 			FEarlyReflectionsFast EarlyReflections;
