Undo changelist 11153240:

2 of 3: Add the Audio Osc Niagara Data Interface.

diff --git a/Engine/Plugins/FX/Niagara/Source/Niagara/Classes/NiagaraDataInterfaceAudioOscilloscope.h b/Engine/Plugins/FX/Niagara/Source/Niagara/Classes/NiagaraDataInterfaceAudioOscilloscope.h
deleted file mode 100644
diff --git a/Engine/Plugins/FX/Niagara/Source/Niagara/Niagara.Build.cs b/Engine/Plugins/FX/Niagara/Source/Niagara/Niagara.Build.cs
--- a/Engine/Plugins/FX/Niagara/Source/Niagara/Niagara.Build.cs
+++ b/Engine/Plugins/FX/Niagara/Source/Niagara/Niagara.Build.cs
@@ -18,8 +18,6 @@ public class Niagara : ModuleRules
                 "Renderer",
                 "JsonUtilities",
 				"Json",
-				"AudioPlatformConfiguration",
-				"SignalProcessing"
             }
         );
 
diff --git a/Engine/Plugins/FX/Niagara/Source/Niagara/Private/NiagaraDataInterfaceAudioOscilloscope.cpp b/Engine/Plugins/FX/Niagara/Source/Niagara/Private/NiagaraDataInterfaceAudioOscilloscope.cpp
deleted file mode 100644
diff --git a/Engine/Source/Runtime/SignalProcessing/Private/MultithreadedPatching.cpp b/Engine/Source/Runtime/SignalProcessing/Private/MultithreadedPatching.cpp
--- a/Engine/Source/Runtime/SignalProcessing/Private/MultithreadedPatching.cpp
+++ b/Engine/Source/Runtime/SignalProcessing/Private/MultithreadedPatching.cpp
@@ -120,15 +120,6 @@ namespace Audio
 	{
 	}
 
-	FPatchInput::FPatchInput(FPatchInput&& Other)
-	{
-		OutputHandle = Other.OutputHandle;
-		Other.OutputHandle.Reset();
-
-		PushCallsCounter = Other.PushCallsCounter;
-		Other.PushCallsCounter = 0;
-	}
-
 	FPatchInput& FPatchInput::operator=(const FPatchInput& Other)
 	{
 		OutputHandle = Other.OutputHandle;
diff --git a/Engine/Source/Runtime/SignalProcessing/Public/DSP/MultithreadedPatching.h b/Engine/Source/Runtime/SignalProcessing/Public/DSP/MultithreadedPatching.h
--- a/Engine/Source/Runtime/SignalProcessing/Public/DSP/MultithreadedPatching.h
+++ b/Engine/Source/Runtime/SignalProcessing/Public/DSP/MultithreadedPatching.h
@@ -83,7 +83,6 @@ namespace Audio
 		FPatchInput(const FPatchOutputStrongPtr& InOutput);
 		FPatchInput(const FPatchInput& Other);
 		FPatchInput& operator=(const FPatchInput& Other);
-		FPatchInput(FPatchInput&& Other);
 
 		/** Default constructed FPatchInput instances will always return -1 for PushAudio and true for IsOutputStillActive. */
 		FPatchInput();
