Fixed edge-case Time Synth crash that occured if the duration of a musical subdivision is smaller than the audio callback size.

diff --git a/Engine/Source/Runtime/SignalProcessing/Private/EventQuantizer.cpp b/Engine/Source/Runtime/SignalProcessing/Private/EventQuantizer.cpp
--- a/Engine/Source/Runtime/SignalProcessing/Private/EventQuantizer.cpp
+++ b/Engine/Source/Runtime/SignalProcessing/Private/EventQuantizer.cpp
@@ -79,7 +79,6 @@ namespace Audio
 		for (int32 Index = 0; Index < (int32)EEventQuantization::Count; ++Index)
 		{
 			FEventQuantizationState& EventState = EventQuantizationStates[Index];
-			EventState.FrameCount = 0;
 			EventState.EventFrameDuration = NumFramesPerBar;
 
 			switch ((EEventQuantization)Index)
@@ -145,6 +144,9 @@ namespace Audio
 				checkf(false, TEXT("Need to update this loop for new quantization enumeration"));
 				break;
 			}
+
+			// Set the events to trigger on the first sample
+			EventState.FrameCount = EventState.EventFrameDuration - 1;
 		}
 
 		BPMQuantizationState.FrameCount = 0;
@@ -195,22 +197,23 @@ namespace Audio
 		BPMQuantizationState.QueuedEvents.Add(MoveTemp(Lambda));
 	}
 
-	void FEventQuantizer::NotifyEventForState(FEventQuantizationState& State, EEventQuantization Type, bool bIsQuantizationEvent, int32 NumFrames)
+	void FEventQuantizer::NotifyEventForState(FEventQuantizationState& State, EEventQuantization Type, bool bIsQuantizationEvent, int32 NumFramesInCallback)
 	{
 		check(State.FrameCount != INDEX_NONE);
 		check(State.EventFrameDuration != INDEX_NONE);
 
-		uint32 NextFrameCount = State.FrameCount + NumFrames;
+		uint32 NextFrameCount = State.FrameCount + NumFramesInCallback;
 
 		bool bResetFrameCount = false;
+		int32 NumEventsFired = 0;
 
 		// Check the event trigger condition
 		if (NextFrameCount >= State.EventFrameDuration)
 		{
 			// Compute the frame offset (basically where in this buffer the event starts)
-			int32 FrameOffset = (int32)(State.EventFrameDuration - State.FrameCount - 1);
-			int32 FramesLeft = NumFrames - FrameOffset;
-			check(FramesLeft >= 0);
+			int32 NextEventFrameStart = (int32)(State.EventFrameDuration - State.FrameCount - 1);
+			int32 PostEventFramesRemaining = NumFramesInCallback - NextEventFrameStart;
+			check(PostEventFramesRemaining >= 0);
 
 			// Support lambda callbacks (events) queuing more events.
 			// Which may need to be executed within a single audio buffer. 
@@ -228,7 +231,7 @@ namespace Audio
 					// We use the copied events since this loop
 					for (TFunction<void(uint32 NumFramesOffset)>& Event : CopiedEvents)
 					{
-						Event(FrameOffset);
+						Event(NextEventFrameStart);
 					}
 				}
 
@@ -241,7 +244,7 @@ namespace Audio
 
 				// Do event listener notifications *with* the frame count with offset accounted for so 
 				// callbacks will get exact frame counts of when this event happened.
-				uint32 FrameCountWithOffset = FrameCount + FrameOffset;
+				uint32 FrameCountWithOffset = FrameCount + NextEventFrameStart;
 				int32 NumBars = FrameCountWithOffset / NumFramesPerBar;
 				int32 NumBeatsFromStart = FrameCountWithOffset / (NumFramesPerBeat - 1);
 				int32 NumBeatsInBar = NumBeatsFromStart % QuantizationSettings.BeatsPerBar;
@@ -254,14 +257,14 @@ namespace Audio
 					}
 				}
 
-				FrameOffset += State.EventFrameDuration;
-				FramesLeft -= State.EventFrameDuration;
-
-			} while (FramesLeft > (int32)State.EventFrameDuration);
+				NextEventFrameStart += State.EventFrameDuration;
+				PostEventFramesRemaining -= State.EventFrameDuration;
+				++NumEventsFired;
+			} while (PostEventFramesRemaining >= 0);
 
 			// Wrap the frame count back to within the event frame duration range
 			// but keep the phase of the frame
-			State.FrameCount = NextFrameCount - (int32)State.EventFrameDuration;
+			State.FrameCount = NextFrameCount - (NumEventsFired * (int32)State.EventFrameDuration);
 		}
 		else
 		{
