Fix for TimeSynth assert caused by uint32 underflow. Edge case where FrameCounter could be -1

diff --git a/Engine/Source/Runtime/SignalProcessing/Private/EventQuantizer.cpp b/Engine/Source/Runtime/SignalProcessing/Private/EventQuantizer.cpp
--- a/Engine/Source/Runtime/SignalProcessing/Private/EventQuantizer.cpp
+++ b/Engine/Source/Runtime/SignalProcessing/Private/EventQuantizer.cpp
@@ -264,7 +264,8 @@ namespace Audio
 
 			// Wrap the frame count back to within the event frame duration range
 			// but keep the phase of the frame
-			State.FrameCount = NextFrameCount - (NumEventsFired * (int32)State.EventFrameDuration);
+			int32 Temp = static_cast<int32>(NextFrameCount) - (NumEventsFired * static_cast<int32>(State.EventFrameDuration));
+			State.FrameCount = FMath::Max(0, Temp);
 		}
 		else
 		{
