Fix for TimeSynth crash w/ absurdly high BPM

diff --git a/Engine/Source/Runtime/SignalProcessing/Private/EventQuantizer.cpp b/Engine/Source/Runtime/SignalProcessing/Private/EventQuantizer.cpp
--- a/Engine/Source/Runtime/SignalProcessing/Private/EventQuantizer.cpp
+++ b/Engine/Source/Runtime/SignalProcessing/Private/EventQuantizer.cpp
@@ -72,9 +72,9 @@ namespace Audio
 		const float BeatDivision = (float)QuantizationSettings.BeatDivision;
 		const float BeatTimeSeconds = 4.0f * QuarterNoteTime / FMath::Max(1.0f, BeatDivision);
 
+		// force a 32second note to be at least a single sample
 		NumFramesPerBeat = (uint32)(BeatTimeSeconds * QuantizationSettings.SampleRate);
-		NumFramesPerBar = QuantizationSettings.BeatsPerBar * NumFramesPerBeat;
-		check(NumFramesPerBar != 0);
+		NumFramesPerBar = FMath::Max(QuantizationSettings.BeatsPerBar * NumFramesPerBeat, 32u);
 
 		for (int32 Index = 0; Index < (int32)EEventQuantization::Count; ++Index)
 		{
