Integrate FFTAlgorithm and ConvolutionAlgorithm work from //UE4/Private-Reverb/...

diff --git a/Engine/Source/Runtime/Android/AudioMixerAndroid/Private/AudioMixerAndroid.cpp b/Engine/Source/Runtime/Android/AudioMixerAndroid/Private/AudioMixerAndroid.cpp
--- a/Engine/Source/Runtime/Android/AudioMixerAndroid/Private/AudioMixerAndroid.cpp
+++ b/Engine/Source/Runtime/Android/AudioMixerAndroid/Private/AudioMixerAndroid.cpp
@@ -7,6 +7,14 @@
 class FAudioMixerModuleAndroid : public IAudioDeviceModule
 {
 public:
+	
+	virtual void StartupModule() override
+	{
+		IAudioDeviceModule::StartupModule();
+
+		FModuleManager::Get().LoadModuleChecked(TEXT("AudioMixerCore"));
+	}
+
 	virtual bool IsAudioMixerModule() const override { return true; }
 
 	virtual Audio::IAudioMixerPlatformInterface* CreateAudioMixerPlatformInterface() override
diff --git a/Engine/Source/Runtime/Apple/AudioMixerAudioUnit/Private/AudioMixerAudioUnit.cpp b/Engine/Source/Runtime/Apple/AudioMixerAudioUnit/Private/AudioMixerAudioUnit.cpp
--- a/Engine/Source/Runtime/Apple/AudioMixerAudioUnit/Private/AudioMixerAudioUnit.cpp
+++ b/Engine/Source/Runtime/Apple/AudioMixerAudioUnit/Private/AudioMixerAudioUnit.cpp
@@ -12,6 +12,14 @@
 class FAudioMixerModuleAudioUnit : public IAudioDeviceModule
 {
 public:
+
+	virtual void StartupModule() override
+	{
+		IAudioDeviceModule::StartupModule();
+
+		FModuleManager::Get().LoadModuleChecked(TEXT("AudioMixerCore"));
+	}
+
 	virtual bool IsAudioMixerModule() const override { return true; }
 
     virtual Audio::IAudioMixerPlatformInterface* CreateAudioMixerPlatformInterface() override
diff --git a/Engine/Source/Runtime/Apple/AudioMixerCoreAudio/Private/AudioMixerCoreAudio.cpp b/Engine/Source/Runtime/Apple/AudioMixerCoreAudio/Private/AudioMixerCoreAudio.cpp
--- a/Engine/Source/Runtime/Apple/AudioMixerCoreAudio/Private/AudioMixerCoreAudio.cpp
+++ b/Engine/Source/Runtime/Apple/AudioMixerCoreAudio/Private/AudioMixerCoreAudio.cpp
@@ -10,6 +10,14 @@
 class FAudioMixerModuleCoreAudio : public IAudioDeviceModule
 {
 public:
+
+	virtual void StartupModule() override
+	{
+		IAudioDeviceModule::StartupModule();
+
+		FModuleManager::Get().LoadModuleChecked(TEXT("AudioMixerCore"));
+	}
+
 	virtual bool IsAudioMixerModule() const override { return true; }
 
 	virtual Audio::IAudioMixerPlatformInterface* CreateAudioMixerPlatformInterface() override
diff --git a/Engine/Source/Runtime/AudioMixer/Private/AudioMixerModule.cpp b/Engine/Source/Runtime/AudioMixer/Private/AudioMixerModule.cpp
--- a/Engine/Source/Runtime/AudioMixer/Private/AudioMixerModule.cpp
+++ b/Engine/Source/Runtime/AudioMixer/Private/AudioMixerModule.cpp
@@ -9,6 +9,8 @@ public:
 
 	virtual void StartupModule() override
 	{
+		FModuleManager::Get().LoadModuleChecked(TEXT("AudioMixerCore"));
+		FModuleManager::Get().LoadModuleChecked(TEXT("SignalProcessing"));
 	}
 };
 
diff --git a/Engine/Source/Runtime/AudioMixerCore/Private/AudioMixerCoreModule.cpp b/Engine/Source/Runtime/AudioMixerCore/Private/AudioMixerCoreModule.cpp
--- a/Engine/Source/Runtime/AudioMixerCore/Private/AudioMixerCoreModule.cpp
+++ b/Engine/Source/Runtime/AudioMixerCore/Private/AudioMixerCoreModule.cpp
@@ -13,6 +13,7 @@ public:
 
 	virtual void StartupModule() override
 	{
+		FModuleManager::Get().LoadModuleChecked(TEXT("SignalProcessing"));
 	}
 };
 
diff --git a/Engine/Source/Runtime/Linux/AudioMixerSDL/Private/AudioMixerSDL.cpp b/Engine/Source/Runtime/Linux/AudioMixerSDL/Private/AudioMixerSDL.cpp
--- a/Engine/Source/Runtime/Linux/AudioMixerSDL/Private/AudioMixerSDL.cpp
+++ b/Engine/Source/Runtime/Linux/AudioMixerSDL/Private/AudioMixerSDL.cpp
@@ -19,6 +19,9 @@ public:
 		FString SDL2_Dll = DllPath + "SDL2.dll";
 		FPlatformProcess::GetDllHandle(*DllPath);
 #endif
+		IAudioDeviceModule::StartupModule();
+
+		FModuleManager::Get().LoadModuleChecked(TEXT("AudioMixerCore"));
 	}
 
 	virtual FAudioDevice* CreateAudioDevice() override
diff --git a/Engine/Source/Runtime/SignalProcessing/Private/AudioFFT.cpp b/Engine/Source/Runtime/SignalProcessing/Private/AudioFFT.cpp
--- a/Engine/Source/Runtime/SignalProcessing/Private/AudioFFT.cpp
+++ b/Engine/Source/Runtime/SignalProcessing/Private/AudioFFT.cpp
@@ -326,7 +326,7 @@ namespace Audio
 
 				// Initialize W of M:
 				float OmegaMReal = FMath::Cos(PI / M2);
-				float OmegaMImag = FMath::Sin(PI / M2);
+				float OmegaMImag = -FMath::Sin(PI / M2);
 
 				for (uint32 j = 0; j < M2; j++)
 				{
@@ -375,7 +375,7 @@ namespace Audio
 
 						float Phase = -2.0f * PI * K / ((float)BitPosition);
 						float TwiddleReal = FMath::Cos(Phase);
-						float TwiddleImag = FMath::Sin(Phase);
+						float TwiddleImag = -FMath::Sin(Phase);
 
 						ComplexMultiply(TwiddleReal, TwiddleImag, OddReal, OddImag, TwiddleReal, TwiddleImag);
 
@@ -504,6 +504,154 @@ namespace Audio
 		}
 	}
 
+	// Implementation of fft algorithm interface
+	class FAudioFFTAlgorithm : public IFFTAlgorithm
+	{
+			const int32 FFTSize;
+			const int32 NumOutputFFTElements;
+			FFTFreqDomainData FreqDomainData;
+
+			AlignedFloatBuffer FreqRealBuffer;
+			AlignedFloatBuffer FreqImagBuffer;
+			
+
+		public:
+
+			FAudioFFTAlgorithm(int32 InFFTSize)
+			:	FFTSize(InFFTSize)
+			,	NumOutputFFTElements((InFFTSize / 2) + 1)
+			{
+				// For freq domain data, we need separate buffers since we are expecting
+				// interleaved [real, imag] data.
+				FreqRealBuffer.AddUninitialized(FFTSize);
+				FreqImagBuffer.AddUninitialized(FFTSize);
+
+				FreqDomainData.OutReal = FreqRealBuffer.GetData();
+				FreqDomainData.OutImag = FreqImagBuffer.GetData();
+			}
+
+			// virtual destructor for inheritance.
+			virtual ~FAudioFFTAlgorithm() {}
+
+			// Number of elements in FFT
+			virtual int32 Size() const override { return FFTSize; }
+
+			// Scaling applied when performing forward FFT.
+			virtual EFFTScaling ForwardScaling() const { return EFFTScaling::None; }
+
+			/* Scaling applied when performing inverse FFT. */
+			virtual EFFTScaling InverseScaling() const { return EFFTScaling::None; }
+
+			
+			// ForwardRealToComplex
+			// InReal - Array of floats to input into fourier transform. Must have FFTSize() elements.
+			// OutComplex - Array of floats to store output of fourier transform. Must have (FFTSize() + 2) float elements which represent ((FFTSize() / 2) + 1) complex numbers in interleaved format.
+			virtual void ForwardRealToComplex(const float* RESTRICT InReal, float* RESTRICT OutComplex) override
+			{
+				const FFTTimeDomainData TimeDomainData = {const_cast<float*>(InReal), FFTSize};
+				PerformFFT(TimeDomainData, FreqDomainData);
+
+				// Convert FFT output data to interleaved format.
+				int32 OutPos = 0;
+				for (int32 i = 0; i < NumOutputFFTElements; i++)
+				{
+					OutComplex[OutPos] = FreqDomainData.OutReal[i];
+					OutPos++;
+					OutComplex[OutPos] = FreqDomainData.OutImag[i];
+					OutPos++;
+				}
+			}
+
+			// InverseComplexToReal
+			// InComplex - Array of floats to input into inverse fourier transform. Must have (FFTSize() + 2) float elements which represent ((FFTSize() / 2) + 1) complex numbers in interleaved format.
+			// OutReal - Array of floats to store output of inverse fourier transform. Must have FFTSize() elements.
+			virtual void InverseComplexToReal(const float* RESTRICT InComplex, float* RESTRICT OutReal) override
+			{
+
+				// For the complex data the phase must be flipped for negative frequencies (or frequencies above nyquist depending on the way you think about it.)
+				
+				// Copy from 0Hz -> Nyquist
+				int32 ComplexPos = 0;
+				for (int32 i = 0; i < NumOutputFFTElements; i++)
+				{
+					FreqDomainData.OutReal[i] = InComplex[ComplexPos];
+					ComplexPos++;
+					FreqDomainData.OutImag[i] = InComplex[ComplexPos];
+					ComplexPos++;
+				}
+
+				// Perform mirror
+				int32 InFreqPos = NumOutputFFTElements - 2;
+				for (int32 MirrorPos = NumOutputFFTElements; MirrorPos < FFTSize; MirrorPos++)
+				{
+					FreqDomainData.OutReal[MirrorPos] = FreqDomainData.OutReal[InFreqPos];
+					FreqDomainData.OutImag[MirrorPos] = -FreqDomainData.OutImag[InFreqPos];
+					InFreqPos--;
+				}
+
+				FFTTimeDomainData TimeDomainData = {OutReal, FFTSize};
+				PerformIFFT(FreqDomainData, TimeDomainData);
+			}
+
+			// BatchForwardRealToComplex
+			virtual void BatchForwardRealToComplex(int32 InCount, const float* const RESTRICT InReal[], float* RESTRICT OutComplex[]) override
+			{
+				for (int32 i = 0; i < InCount; i++)
+				{
+					ForwardRealToComplex(InReal[i], OutComplex[i]);
+				}
+			}
+
+			// BatchInverseComplexToReal
+			virtual void BatchInverseComplexToReal(int32 InCount, const float* const RESTRICT InComplex[], float* RESTRICT OutReal[]) override
+			{
+				for (int32 i = 0; i < InCount; i++)
+				{
+					InverseComplexToReal(InComplex[i], OutReal[i]);
+				}
+			}
+	};
+
+
+	// FFT Algorithm factory for this FFT implementation
+	FAudioFFTAlgorithmFactory::~FAudioFFTAlgorithmFactory() {}
+
+	// Name of this fft algorithm factory. 
+	FName FAudioFFTAlgorithmFactory::GetFactoryName() const 
+	{
+		static const FName FactoryName = FName(TEXT("OriginalFFT_Deprecated"));
+		return FactoryName;
+	}
+
+	// If true, this implementation uses hardware acceleration.
+	bool FAudioFFTAlgorithmFactory::IsHardwareAccelerated() const
+	{
+		return false;
+	}
+
+	// If true, this implementation requires input and output arrays to be 128 bit aligned.
+	bool FAudioFFTAlgorithmFactory::Expects128BitAlignedArrays() const
+	{
+		return false;
+	}
+
+	// Returns true if the input settings are supported by this factory.
+	bool FAudioFFTAlgorithmFactory::AreFFTSettingsSupported(const FFFTSettings& InSettings) const 
+	{
+		return (InSettings.Log2Size > 1) && (InSettings.Log2Size < 30);
+	}
+
+	// Create a new FFT algorithm.
+	TUniquePtr<IFFTAlgorithm> FAudioFFTAlgorithmFactory::NewFFTAlgorithm(const FFFTSettings& InSettings)
+	{
+		check(AreFFTSettingsSupported(InSettings));
+
+		// equivalent of 2^(InSettings.Log2Size)
+		int32 FFTSize = 1 << InSettings.Log2Size;
+
+		return MakeUnique<FAudioFFTAlgorithm>(FFTSize);
+	}
+
 	void ComputePowerSpectrumNoScaling(const FFTFreqDomainData& InFrequencyData, int32 FFTSize, AlignedFloatBuffer& OutBuffer)
 	{
 		check((FFTSize % 2) == 0);
diff --git a/Engine/Source/Runtime/SignalProcessing/Private/ConvolutionAlgorithm.cpp b/Engine/Source/Runtime/SignalProcessing/Private/ConvolutionAlgorithm.cpp
new file mode 100644
--- /dev/null
+++ b/Engine/Source/Runtime/SignalProcessing/Private/ConvolutionAlgorithm.cpp
@@ -0,0 +1,98 @@
+// Copyright Epic Games, Inc. All Rights Reserved.
+
+#include "DSP/ConvolutionAlgorithm.h"
+#include "SignalProcessingModule.h"
+#include "CoreMinimal.h"
+#include "Features/IModularFeature.h"
+#include "Features/IModularFeatures.h"
+
+namespace Audio
+{
+	IConvolutionAlgorithmFactory::~IConvolutionAlgorithmFactory()
+	{}
+
+	const FName IConvolutionAlgorithmFactory::GetModularFeatureName()
+	{
+		static const FName ModularFeatureName = FName(TEXT("AudioConvolutionAlgorithmFactory"));
+		return ModularFeatureName;
+	}
+
+	const FName FConvolutionFactory::AnyAlgorithmFactory = FName(TEXT("AnyAlgorithmFactory"));
+
+	TUniquePtr<IConvolutionAlgorithm> FConvolutionFactory::NewConvolutionAlgorithm(const FConvolutionSettings& InSettings, const FName& InAlgorithmFactoryName)
+	{
+		// Get all IConvolutionAlgorithm factories. 
+		TArray<IConvolutionAlgorithmFactory*> Factories = IModularFeatures::Get().GetModularFeatureImplementations<IConvolutionAlgorithmFactory>(IConvolutionAlgorithmFactory::GetModularFeatureName());
+
+		// Remove null factories
+		Factories = Factories.FilterByPredicate([InAlgorithmFactoryName](const IConvolutionAlgorithmFactory* Factory) 
+		{
+			return (nullptr != Factory);
+		});
+
+		if (0 == Factories.Num())
+		{
+			UE_LOG(LogSignalProcessing, Warning, TEXT("No registered IConvolutionAlgorithmFactories exist."));
+			return TUniquePtr<IConvolutionAlgorithm>();
+		}
+
+		if (AnyAlgorithmFactory != InAlgorithmFactoryName)
+		{
+			// Filter factories if one has been specified. 
+			Factories = Factories.FilterByPredicate([InAlgorithmFactoryName](const IConvolutionAlgorithmFactory* Factory) 
+			{
+				return (InAlgorithmFactoryName == Factory->GetFactoryName());
+			});
+
+			if (0 == Factories.Num())
+			{
+				UE_LOG(LogSignalProcessing, Warning, TEXT("No IConvolutionAlgorithmFactories named \"%s\" found."), *InAlgorithmFactoryName.ToString());
+				return TUniquePtr<IConvolutionAlgorithm>();
+			}
+		}
+
+		if (!InSettings.bEnableHardwareAcceleration)
+		{
+			// Remove factories with hardware acceleration if it has been disabled. 
+			Factories = Factories.FilterByPredicate([InAlgorithmFactoryName](const IConvolutionAlgorithmFactory* Factory) 
+			{
+				return !Factory->IsHardwareAccelerated();
+			});
+		}
+
+		// Sort factories by preference. 
+		Factories.Sort([](const IConvolutionAlgorithmFactory& InFactoryA, const IConvolutionAlgorithmFactory& InFactoryB) 
+		{
+			// If InFactoryA has hardware acceleration and InFactoryB does not, then InFactoryA goes earlier.
+			if (InFactoryA.IsHardwareAccelerated() && !InFactoryB.IsHardwareAccelerated())
+			{
+				return true;
+			}
+
+			// InFactoryA does not need to be placed before InFactoryB
+			return false;
+		});
+
+
+		// Find first algorithm that can handle input settings
+		for (IConvolutionAlgorithmFactory* Factory : Factories)
+		{
+			if (Factory->AreConvolutionSettingsSupported(InSettings))
+			{
+				TUniquePtr<IConvolutionAlgorithm> ConvolutionAlgorithm = Factory->NewConvolutionAlgorithm(InSettings);
+
+				if (!ConvolutionAlgorithm.IsValid())
+				{
+					UE_LOG(LogSignalProcessing, Warning, TEXT("IConvolutionAlgorithmFactory \"%s\" failed to create IConvolutionAlgorithm despite supporting ConvolutionReverb Settings."));
+				}
+				else
+				{
+					return ConvolutionAlgorithm;
+				}
+			}
+		}
+
+		UE_LOG(LogSignalProcessing, Warning, TEXT("No IConvolutionAlgorithmFactories can create IConvolutionAlgorithm for given settings."));
+		return TUniquePtr<IConvolutionAlgorithm>();
+	}
+}
diff --git a/Engine/Source/Runtime/SignalProcessing/Private/FFTAlgorithm.cpp b/Engine/Source/Runtime/SignalProcessing/Private/FFTAlgorithm.cpp
new file mode 100644
--- /dev/null
+++ b/Engine/Source/Runtime/SignalProcessing/Private/FFTAlgorithm.cpp
@@ -0,0 +1,160 @@
+// Copyright Epic Games, Inc. All Rights Reserved.
+
+#include "DSP/FFTAlgorithm.h"
+#include "SignalProcessingModule.h"
+#include "CoreMinimal.h"
+#include "Features/IModularFeature.h"
+#include "Features/IModularFeatures.h"
+
+namespace Audio
+{
+	// Return factories which support settings and factory name.
+	static const TArray<IFFTAlgorithmFactory*> GetPrioritizedSupportingFactories(const FFFTSettings& InSettings, const FName& InAlgorithmFactoryName)
+	{
+
+		// Get all IFFTAlgorithm factories. 
+		TArray<IFFTAlgorithmFactory*> Factories = IModularFeatures::Get().GetModularFeatureImplementations<IFFTAlgorithmFactory>(IFFTAlgorithmFactory::GetModularFeatureName());
+
+		// Remove null factories
+		Factories = Factories.FilterByPredicate([InAlgorithmFactoryName](const IFFTAlgorithmFactory* Factory) 
+		{
+			return (nullptr != Factory);
+		});
+
+		if (0 == Factories.Num())
+		{
+			UE_LOG(LogSignalProcessing, Warning, TEXT("No registered IFFTAlgorithmFactories exist."));
+			return Factories;
+		}
+
+		if (FFFTFactory::AnyAlgorithmFactory != InAlgorithmFactoryName)
+		{
+			// Filter factories if one has been specified. 
+			Factories = Factories.FilterByPredicate([InAlgorithmFactoryName](const IFFTAlgorithmFactory* Factory) 
+			{
+				return (InAlgorithmFactoryName == Factory->GetFactoryName());
+			});
+
+			if (0 == Factories.Num())
+			{
+				UE_LOG(LogSignalProcessing, Warning, TEXT("No IFFTAlgorithmFactories named \"%s\" found."), *InAlgorithmFactoryName.ToString());
+				return Factories;
+			}
+		}
+
+		if (!InSettings.bEnableHardwareAcceleration)
+		{
+			// Remove factories with hardware acceleration if it has been disabled. 
+			Factories = Factories.FilterByPredicate([InAlgorithmFactoryName](const IFFTAlgorithmFactory* Factory) 
+			{
+				return !Factory->IsHardwareAccelerated();
+			});
+		}
+
+		if (!InSettings.bArrays128BitAligned)
+		{
+			// Remove factories requiring aligned input/output arrays if it has been disabled.
+			Factories = Factories.FilterByPredicate([InAlgorithmFactoryName](const IFFTAlgorithmFactory* Factory) 
+			{
+				return !Factory->Expects128BitAlignedArrays();
+			});
+		}
+
+		// Filter by whether factory supports settings.
+		Factories = Factories.FilterByPredicate([InSettings](const IFFTAlgorithmFactory* Factory)
+		{
+			return Factory->AreFFTSettingsSupported(InSettings);
+		});
+
+		// We have two software versions, but should prefer the vectorized version
+		// over the non-vectorized version.
+		static const TArray<FName> InternalFFTFActoryPriority = {
+			FName(TEXT("FVectorFFTFactory")),
+			FName(TEXT("OriginalFFT_Deprecated"))
+		};
+
+		// Sort factories by preference. 
+		Factories.Sort([](const IFFTAlgorithmFactory& InFactoryA, const IFFTAlgorithmFactory& InFactoryB) 
+		{
+			// If InFactoryA has hardware acceleration and InFactoryB does not, then InFactoryA goes earlier.
+			if (InFactoryA.IsHardwareAccelerated() && !InFactoryB.IsHardwareAccelerated())
+			{
+				return true;
+			}
+
+			// If InFactoryA uses aligned arrays and InFactoryB does not, then InFactoryA goes earlier.
+			if (InFactoryA.Expects128BitAlignedArrays() && !InFactoryB.Expects128BitAlignedArrays())
+			{
+				return true;
+			}
+
+			// Check internal priority of software implementations.
+			int32 PriorityIndexA = InternalFFTFActoryPriority.Find(InFactoryA.GetFactoryName());
+			int32 PriorityIndexB = InternalFFTFActoryPriority.Find(InFactoryB.GetFactoryName());
+
+			if ((PriorityIndexA != INDEX_NONE) && (PriorityIndexB != INDEX_NONE))
+			{
+				return PriorityIndexA < PriorityIndexB;
+				
+			}
+
+			// InFactoryA does not need to be placed before InFactoryB
+			return false;
+		});
+		
+		return Factories;
+	}
+
+	/** virtual destructor for inheritance. */
+	IFFTAlgorithm::~IFFTAlgorithm()
+	{
+	}
+
+	/** Name of modular feature for FFT factory.  */
+	const FName IFFTAlgorithmFactory::GetModularFeatureName()
+	{
+		static const FName ModularFeatureName = FName(TEXT("AudioFFTAlgorithmFactory"));
+		return ModularFeatureName;
+	}
+
+	/** This denotes that no specific IFFTAlgorithmFactory is desired. */
+	const FName FFFTFactory::AnyAlgorithmFactory = FName(TEXT("AnyAlgorithmFactory"));
+
+	/** NewFFTAlgorithm
+	 *
+	 * Creates and returns a new FFTAlgorithm. 
+	 *
+	 * @param InSettings - The settings used to create the FFT algorithm.
+	 * @param InAlgorithmFactoryName - If not equal to FFFTFactory::AnyAlgorithmFactory, will only uses FFT algorithm facotry if IFFTAlgorithmFactory::GetFactoryName() equals InAlgorithmFactoryName.
+	 * @return A TUniquePtr<IFFTAlgorithm> to the created FFT. This pointer can be invalid if an error occured or the fft algorithm could not be created.
+	 */
+	TUniquePtr<IFFTAlgorithm> FFFTFactory::NewFFTAlgorithm(const FFFTSettings& InSettings, const FName& InAlgorithmFactoryName)
+	{
+		TArray<IFFTAlgorithmFactory*> Factories = GetPrioritizedSupportingFactories(InSettings, InAlgorithmFactoryName);
+
+		// Find first algorithm which creates the an algorithm from the input settings. 
+		for (IFFTAlgorithmFactory* Factory : Factories)
+		{
+			TUniquePtr<IFFTAlgorithm> FFTAlgorithm = Factory->NewFFTAlgorithm(InSettings);
+
+			if (!FFTAlgorithm.IsValid())
+			{
+				UE_LOG(LogSignalProcessing, Warning, TEXT("IFFTAlgorithmFactory \"%s\" failed to create IFFTAlgorithm despite supporting FFT Settings."));
+			}
+			else
+			{
+				return FFTAlgorithm;
+			}
+		}
+
+
+		UE_LOG(LogSignalProcessing, Warning, TEXT("No IFFTAlgorithmFactories can create IFFTAlgorithm for given settings."));
+		return TUniquePtr<IFFTAlgorithm>();
+	}
+
+	bool FFFTFactory::AreFFTSettingsSupported(const FFFTSettings& InSettings, const FName& InAlgorithmFactoryName)
+	{
+		TArray<IFFTAlgorithmFactory*> Factories = GetPrioritizedSupportingFactories(InSettings, InAlgorithmFactoryName);
+		return Factories.Num() > 0;
+	}
+}
diff --git a/Engine/Source/Runtime/SignalProcessing/Private/SignalProcessingModule.cpp b/Engine/Source/Runtime/SignalProcessing/Private/SignalProcessingModule.cpp
--- a/Engine/Source/Runtime/SignalProcessing/Private/SignalProcessingModule.cpp
+++ b/Engine/Source/Runtime/SignalProcessing/Private/SignalProcessingModule.cpp
@@ -2,16 +2,46 @@
 
 #include "SignalProcessingModule.h"
 #include "Modules/ModuleManager.h"
+#include "Features/IModularFeatures.h"
+#include "DSP/FFTAlgorithm.h"
+#include "DSP/ConvolutionAlgorithm.h"
+#include "DSP/AudioFFT.h"
+#include "VectorFFT.h"
+#include "UniformPartitionConvolution.h"
 
-class FSignalProcessingModule : public IModuleInterface
+namespace Audio
 {
-public:
-
-	virtual void StartupModule() override
+	class FSignalProcessingModule : public IModuleInterface
 	{
-	}
-};
+		FAudioFFTAlgorithmFactory AudioFFTAlgorithmFactory;
+		FVectorFFTFactory VectorFFTAlgorithmFactory;
+
+		FUniformPartitionConvolutionFactory UniformPartitionConvolutionFactory;
+
+	public:
+
+		virtual void StartupModule() override
+		{
+			// FFT factories to register
+			IModularFeatures::Get().RegisterModularFeature(IFFTAlgorithmFactory::GetModularFeatureName(), &AudioFFTAlgorithmFactory);
+			IModularFeatures::Get().RegisterModularFeature(IFFTAlgorithmFactory::GetModularFeatureName(), &VectorFFTAlgorithmFactory);
+
+			// Convolution factories to register
+			IModularFeatures::Get().RegisterModularFeature(IConvolutionAlgorithmFactory::GetModularFeatureName(), &UniformPartitionConvolutionFactory);
+		}
+
+		virtual void ShutdownModule() override
+		{
+			// FFT Factories to unregister
+			IModularFeatures::Get().UnregisterModularFeature(IFFTAlgorithmFactory::GetModularFeatureName(), &AudioFFTAlgorithmFactory);
+			IModularFeatures::Get().UnregisterModularFeature(IFFTAlgorithmFactory::GetModularFeatureName(), &VectorFFTAlgorithmFactory);
+
+			// Convolution factories to unregister
+			IModularFeatures::Get().UnregisterModularFeature(IConvolutionAlgorithmFactory::GetModularFeatureName(), &UniformPartitionConvolutionFactory);
+		}
+	};
+}
 
 DEFINE_LOG_CATEGORY(LogSignalProcessing);
 
-IMPLEMENT_MODULE(FSignalProcessingModule, SignalProcessing);
+IMPLEMENT_MODULE(Audio::FSignalProcessingModule, SignalProcessing);
diff --git a/Engine/Source/Runtime/SignalProcessing/Private/UniformPartitionConvolution.cpp b/Engine/Source/Runtime/SignalProcessing/Private/UniformPartitionConvolution.cpp
new file mode 100644
--- /dev/null
+++ b/Engine/Source/Runtime/SignalProcessing/Private/UniformPartitionConvolution.cpp
@@ -0,0 +1,628 @@
+// Copyright Epic Games, Inc. All Rights Reserved.
+
+#include "UniformPartitionConvolution.h"
+#include "CoreMinimal.h"
+#include "SignalProcessingModule.h"
+#include "DSP/BufferVectorOperations.h"
+#include "DSP/ConvolutionAlgorithm.h"
+#include "DSP/FFTAlgorithm.h"
+
+///////////////////////////////////////////////////////////////////////////////
+// Note on OVERLAP_SAVE
+//
+// This algorithm is implemented using the overlap save method of convolution.
+//
+// The block size is set to FFT->Size() / 2 samples.
+//
+// Inputs FFTs are performed on the input samples using a rectaangular sliding 
+// window of 2 * BlockSize samples with a hop of BlockSize samples.
+//
+// Impulse response FFTs 2 * BlockSize with BlockSize zeropadded samples. 
+//
+// Output inverse FFTs produce 2 * BlockSize samples, but the first BlockSize
+// samples are discarded because they are time aliased. 
+///////////////////////////////////////////////////////////////////////////////
+
+namespace Audio
+{
+	FUniformPartitionConvolution::FUniformPartitionConvolution(const FUniformPartitionConvolutionSettings& InSettings, FSharedFFTRef InFFTAlgorithm)
+	:	Settings(InSettings)
+	,	BlockSize(0)
+	,	NumFFTOutputFloats(0)
+	,	NumBlocks(0)
+	,	FFTAlgorithm(InFFTAlgorithm)
+	{
+		// Input to FFT has to be zero padded in order to prevent aliasing from circular 
+		// convolution. Hence, block size is half the fft size. 
+		BlockSize = FFTAlgorithm->Size() / 2;
+
+		check(BlockSize > 0);
+
+		// Number of floats in fft output. 
+		NumFFTOutputFloats = FFTAlgorithm->NumOutputFloats();
+
+		TransformedAndScaledInput.AddUninitialized(NumFFTOutputFloats);
+
+		NumBlocks = Settings.MaxNumImpulseResponseSamples / BlockSize;
+		// Add extra block if needed. Accounting for integer division rounding down.
+		NumBlocks += Settings.MaxNumImpulseResponseSamples % BlockSize == 0 ? 0 : 1;
+
+		// Create inputs 
+		for (int32 i = 0; i < Settings.NumInputChannels; i++)
+		{
+			Inputs.Emplace(FFTAlgorithm);
+		}
+
+		// Create transforms
+		for (int32 i = 0; i < Settings.NumImpulseResponses; i++)
+		{
+			ImpulseResponses.Emplace(NumBlocks, FFTAlgorithm);
+		}
+
+		// Create outputs 
+		for (int32 i = 0; i < Settings.NumOutputChannels; i++)
+		{
+			Outputs.Emplace(NumBlocks, FFTAlgorithm);
+		}
+		IsOutputZero.Init(true, Settings.NumOutputChannels);
+	}
+
+	FUniformPartitionConvolution::~FUniformPartitionConvolution()
+	{
+	}
+
+	/** Returns the number of samples in an audio block. */
+	int32 FUniformPartitionConvolution::GetNumSamplesInBlock() const
+	{
+		return BlockSize;
+	}
+
+	/** Returns number of audio inputs. */
+	int32 FUniformPartitionConvolution::GetNumAudioInputs() const
+	{
+		return Settings.NumInputChannels;
+	}
+
+	/** Returns number of audio outputs. */
+	int32 FUniformPartitionConvolution::GetNumAudioOutputs() const
+	{
+		return Settings.NumOutputChannels;
+	}
+
+	/** Process one block of audio.
+	*
+	* InSamples is processed by the impulse responses. The output is placed in OutSamples.
+	*
+	* @params InSamples - A 2D array of input deinterleaved audio samples. InSamples[GetNumAudioInputs()][GetNumSamplesInBlock()]
+	* @params OutSamples - A 2D array of output deinterleaved audio samples. OutSamples[GetNumAudioOutputs()][GetNumSamplesInBlock()]
+	*
+	*/
+	void FUniformPartitionConvolution::ProcessAudioBlock(const float* const InSamples[], float* const OutSamples[])
+	{
+		
+		// Add input audio to input channels
+		for (int32 i = 0; i < Settings.NumInputChannels; i++)
+		{
+			Inputs[i].PushBlock(InSamples[i]);
+		}
+		
+		// Apply filtering for gains above zero
+		for (auto GainMatrixEntry : NonZeroGains)
+		{
+			// Get input, transform and output.
+			const FInputTransformOutputGroup& Group = GainMatrixEntry.Key;
+
+			const FInput& Input = Inputs[Group.GetInputIndex()];
+			const FImpulseResponse& Transform = ImpulseResponses[Group.GetTransformIndex()];
+			FOutput& Output = Outputs[Group.GetOutputIndex()];
+
+			// Get gain
+			const float Gain = GainMatrixEntry.Value;
+
+			// Only process the number of blocks needed to process the specific transform.
+			const int32 NumBlocksToProcess = Transform.GetNumActiveBlocks();
+
+			if (NumBlocksToProcess > 0)
+			{
+				// Scale input by gain
+				VectorMultiplyByConstant(Input.GetTransformedBlock(), Gain, TransformedAndScaledInput);
+
+				// Multiply transformed input by each block of transformed impulse response.
+				// Accumulate the result of that multiplication in the appropriate output block.
+				for (int32 i = 0; i < NumBlocksToProcess; i++)
+				{
+					VectorComplexMultiplyAdd(TransformedAndScaledInput, Transform.GetTransformedBlock(i), Output.GetTransformedBlock(i));
+				}
+			}
+		}
+
+		// Pop blocks off of outputs.
+		for (int32 i = 0; i < Settings.NumOutputChannels; i++)
+		{
+			if (IsOutputZero[i])
+			{
+				FMemory::Memset(OutSamples[i], 0, sizeof(float) * BlockSize);
+
+				// This avoids calling inverse FFT when output is zero due to 
+				// no non-zero gains. 
+				Outputs[i].PopBlock();
+			}
+			else
+			{
+				Outputs[i].PopBlock(OutSamples[i]);
+			}
+		}
+	}
+
+	/** Reset internal history buffers . */
+	void FUniformPartitionConvolution::ResetAudioHistory()
+	{
+		for (int32 i = 0; i < Settings.NumInputChannels; i++)
+		{
+			Inputs[i].Reset();
+		}
+
+		for (int32 i = 0; i < Settings.NumOutputChannels; i++)
+		{
+			Outputs[i].Reset();
+		}
+	}
+
+	/** Maximum supported length of impulse response. */
+	int32 FUniformPartitionConvolution::GetMaxNumImpulseResponseSamples() const
+	{
+		return Settings.MaxNumImpulseResponseSamples;
+	}
+
+	/** Return the number of impulse responses. */
+	int32 FUniformPartitionConvolution::GetNumImpulseResponses() const
+	{
+		return Settings.NumImpulseResponses;
+	}
+
+	/** Return the number of samples in an impulse response. */
+	int32 FUniformPartitionConvolution::GetNumImpulseResponseSamples(int32 InImpulseResponseIndex) const
+	{
+		return ImpulseResponses[InImpulseResponseIndex].GetNumImpulseResponseSamples();
+	}
+
+	/** Set impulse response values. */
+	void FUniformPartitionConvolution::SetImpulseResponse(int32 InImpulseResponseIndex, const float* InSamples, int32 InNumSamples)
+	{
+		check(InNumSamples <= Settings.MaxNumImpulseResponseSamples);
+		ImpulseResponses[InImpulseResponseIndex].SetImpulseResponse(InSamples, InNumSamples);
+	}
+
+	/** Sets the gain between an audio input, impulse response and audio output.
+	*
+	* ([audio inputs] * [impulse responses]) x [gain matrix] = [audio outputs]
+	*/
+	void FUniformPartitionConvolution::SetMatrixGain(int32 InAudioInputIndex, int32 InImpulseResponseIndex, int32 InAudioOutputIndex, float InGain)
+	{
+		// Create index into 3D gain matrix
+		const FInputTransformOutputGroup Group(InAudioInputIndex, InImpulseResponseIndex, InAudioOutputIndex);
+
+		if (InGain == 0.f)
+		{
+			// Remove gain entry if it's zero
+			NonZeroGains.Remove(Group);
+
+			// First, assume if setting gain to zero that no audio is sent to output.
+			IsOutputZero[InAudioOutputIndex] = true;
+
+			// Check that no non-zero gain is associated with this output.
+			for (auto GainMapEntry : NonZeroGains)
+			{
+				if (GainMapEntry.Key.GetOutputIndex() == InAudioOutputIndex)
+				{
+					// There is a non-zero gain associated with this output. 
+					IsOutputZero[InAudioOutputIndex] = false;
+					break;
+				}
+			}
+		}
+		else
+		{
+			NonZeroGains.Add(Group, InGain);
+			IsOutputZero[InAudioOutputIndex] = false;
+		}
+	}
+
+	/** Gets the gain between an audio input, impulse response and audio output.
+	*
+	* ([audio inputs] * [impulse responses]) x [gain matrix] = [audio outputs]
+	*/
+	float FUniformPartitionConvolution::GetMatrixGain(int32 InAudioInputIndex, int32 InImpulseResponseIndex, int32 InAudioOutputIndex) const
+	{
+		float OutGain = 0.f;
+
+		const FInputTransformOutputGroup Group(InAudioInputIndex, InImpulseResponseIndex, InAudioOutputIndex);
+		if (NonZeroGains.Contains(Group))
+		{
+			OutGain = NonZeroGains[Group];
+		}
+
+		return OutGain;
+	}
+
+	// If you bitwise and with this, it gives you (Num = Num - (Num % 4)). This is an easy way
+	// to determine how many values can be SIMD'd.
+	const int32 FUniformPartitionConvolution::NumSimdMask = 0xFFFFFFFC;
+
+	void FUniformPartitionConvolution::VectorComplexMultiplyAdd(const AlignedFloatBuffer& InA, const AlignedFloatBuffer& InB, AlignedFloatBuffer& Out) const
+	{
+		check(InA.Num() == InB.Num());
+		check(Out.Num() == InA.Num());
+
+		const int32 Num = InA.Num();
+		const int32 NumSimd = Num & NumSimdMask;
+
+		// Complex numbers are stored as [real_0, complex_0, real_1, complex_1, ... real_N, complex_N]
+		// So we final amount must be evenly divisble by 2.
+		check(NumSimd % 2 == 0);
+
+		const float* InAData = InA.GetData();
+		const float* InBData = InB.GetData();
+		float* OutData = Out.GetData();
+
+		const VectorRegister SignFlip = MakeVectorRegister(-1.f, 1.f, -1.f, 1.f);
+
+		for (int32 i = 0; i < NumSimd; i += 4)
+		{
+			// Complex multiply add
+			// Nr = real component of Nth number
+			// Ni = imaginary component of Nth number
+			//
+			//
+			// The input is then 
+			// A1r A1i A2r A2i
+			// B1r B1i B2r B2i
+
+			// VectorA = A1r A1i A2r A2i
+			VectorRegister VectorInA = VectorLoadAligned(&InAData[i]);
+			// Temp12 = A1i A1r A2i A2r
+			VectorRegister Temp1 = VectorSwizzle(VectorInA, 1, 0, 3, 2);
+
+			// VectorB = B1r B1i B2r B2i
+			VectorRegister VectorInB = VectorLoadAligned(&InBData[i]);
+			// Temp2 = B1r B1r B2r B2r
+			VectorRegister Temp2 = VectorSwizzle(VectorInB, 0, 0, 2, 2);
+			// Temp3 = B1i B1i B2i B2i
+			VectorRegister Temp3 = VectorSwizzle(VectorInB, 1, 1, 3, 3);
+
+
+			// VectorA = A1rB1r, A1iB1r, A2rB2r, A2iB2r
+			VectorInA = VectorMultiply(VectorInA, Temp2);
+
+			// Temp1 = A1iB1i, A1rB1i, A2iB2i, A2rb2i
+			Temp1 = VectorMultiply(Temp1, Temp3);
+
+			// Temp1 = -A1iB1i, A1rB1i, -A2iB2i, A2rb2i
+			Temp1 = VectorMultiply(Temp1, SignFlip);
+
+
+			// Temp1 = A1rB1r - A1iB1i, A1iB1r + A1rB1i, A2rB2r - A2iB2i, A2iB2r + A2rB2i
+			Temp1 = VectorAdd(VectorInA, Temp1);
+
+			// VectorOut = O1r + A1rB1r - A1iB1i, O1i + A1iB1r + A1rB1i, O2r + A2rB2r - A2iB2i, O2i + A2iB2r + A2rB2i
+			VectorRegister VectorOut = VectorLoadAligned(&OutData[i]);
+			VectorOut = VectorAdd(Temp1, VectorOut);
+
+			VectorStoreAligned(VectorOut, &OutData[i]);
+		}
+
+		for (int32 i = NumSimd; i < Num; i += 2)
+		{
+			// Real output
+			OutData[i] += (InAData[i] * InBData[i]) - (InAData[i + 1] * InBData[i + 1]);
+			// Imaginary output
+			OutData[i + 1] += (InAData[i + 1] * InBData[i]) + (InAData[i] * InBData[i + 1]);
+		}
+	}
+
+	// Multiply aligned buffer by constant gain.
+	void FUniformPartitionConvolution::VectorMultiplyByConstant(const AlignedFloatBuffer& InBuffer, float InConstant, AlignedFloatBuffer& OutBuffer) const
+	{
+		check(InBuffer.Num() == OutBuffer.Num());
+
+		VectorRegister VectorConstant = MakeVectorRegister(InConstant, InConstant, InConstant, InConstant);
+
+		const int32 Num = InBuffer.Num();
+		const int32 NumSimd = Num & NumSimdMask;
+
+		const float* InData = InBuffer.GetData();
+		float* OutData = OutBuffer.GetData();
+
+		for (int32 i = 0; i < NumSimd; i += 4)
+		{
+			VectorRegister VectorIn = VectorLoadAligned(&InData[i]);
+
+			VectorRegister VectorOut = VectorMultiply(VectorIn, VectorConstant);
+			
+			VectorStoreAligned(VectorOut, &OutData[i]);
+		}
+
+		for (int32 i = NumSimd; i < Num; i++)
+		{
+			OutData[i] = InData[i] * InConstant;
+		}
+	}
+
+	FUniformPartitionConvolution::FInput::FInput(FSharedFFTRef InFFTAlgorithm)
+	:	BlockSize(InFFTAlgorithm->Size() / 2)
+	,	FFTAlgorithm(InFFTAlgorithm)
+	{
+		check(BlockSize > 0);
+
+		InputBuffer.AddZeroed(FFTAlgorithm->NumInputFloats());
+		OutputBuffer.AddUninitialized(FFTAlgorithm->NumOutputFloats());
+	}
+
+	void FUniformPartitionConvolution::FInput::PushBlock(const float* InSamples)
+	{
+		float* InputData = InputBuffer.GetData();
+		float* OutputData = OutputBuffer.GetData();
+
+		// The two copies here effectively perform a hop of "BlockSize" number of samples
+		// between subsequent calls to FFTAlgorithm->ForwardRealToComplex
+		FMemory::Memcpy(&InputData[BlockSize], InSamples, BlockSize * sizeof(float));
+
+		FFTAlgorithm->ForwardRealToComplex(InputData, OutputData);
+
+		FMemory::Memcpy(InputData, InSamples, BlockSize * sizeof(float));
+	}
+
+	const AlignedFloatBuffer& FUniformPartitionConvolution::FInput::GetTransformedBlock() const
+	{
+		return OutputBuffer;
+	}
+
+	void FUniformPartitionConvolution::FInput::Reset()
+	{
+		FMemory::Memset(InputBuffer.GetData(), 0, sizeof(float) * InputBuffer.Num());
+		FMemory::Memset(OutputBuffer.GetData(), 0, sizeof(float) * OutputBuffer.Num());
+	}
+
+	FUniformPartitionConvolution::FOutput::FOutput(int32 InNumBlocks, FSharedFFTRef InFFTAlgorithm)
+	:	NumBlocks(InNumBlocks)
+	,	BlockSize(InFFTAlgorithm->Size() / 2)
+	,	NumFFTOutputFloats(InFFTAlgorithm->NumOutputFloats())
+	,	FFTAlgorithm(InFFTAlgorithm)
+	,	HeadBlockIndex(0)
+	{
+		Blocks.AddDefaulted(NumBlocks);
+
+		for (int32 i = 0; i < NumBlocks; i++)
+		{
+			Blocks[i].AddZeroed(NumFFTOutputFloats);
+		}
+
+		OutputBuffer.AddZeroed(FFTAlgorithm->NumInputFloats());
+	}
+
+
+	FUniformPartitionConvolution::FOutput::~FOutput()
+	{
+	}
+
+	AlignedFloatBuffer& FUniformPartitionConvolution::FOutput::GetTransformedBlock(int32 InBlockIndex)
+	{
+		// Make block index relative to head index.
+		InBlockIndex = (HeadBlockIndex + InBlockIndex) % NumBlocks;
+
+		return Blocks[InBlockIndex];
+	}
+
+	void FUniformPartitionConvolution::FOutput::PopBlock()
+	{
+		const int32 TailBlockIndex = HeadBlockIndex;
+
+		// Increment head pointer
+		HeadBlockIndex = (HeadBlockIndex + 1) % NumBlocks;
+
+		// Set tail block to zero
+		AlignedFloatBuffer& TailBuffer = Blocks[TailBlockIndex];
+		FMemory::Memset(TailBuffer.GetData(), 0, NumFFTOutputFloats * sizeof(float));
+
+	}
+
+	void FUniformPartitionConvolution::FOutput::PopBlock(float* OutSamples)
+	{
+		const AlignedFloatBuffer& HeadBuffer = Blocks[HeadBlockIndex];
+		float* OutputData = OutputBuffer.GetData();
+
+		FFTAlgorithm->InverseComplexToReal(HeadBuffer.GetData(), OutputData);
+
+		// Overlap save method of convolution. Only use 2nd half of output buffer
+		FMemory::Memcpy(OutSamples, &OutputData[BlockSize], sizeof(float) * BlockSize);
+
+		// Perform head block increment. 
+		PopBlock();
+	}
+
+	void FUniformPartitionConvolution::FOutput::Reset()
+	{
+		for (int32 i = 0; i < NumBlocks; i++)
+		{
+			FMemory::Memset(Blocks[i].GetData(), 0, sizeof(float) * NumFFTOutputFloats);
+		}
+
+		FMemory::Memset(OutputBuffer.GetData(), 0, sizeof(float) * 2 * BlockSize);
+	}
+
+
+	// Hold FFTs of IR
+	FUniformPartitionConvolution::FImpulseResponse::FImpulseResponse(int32 InNumBlocks, FSharedFFTRef InFFTAlgorithm)
+	:	NumBlocks(InNumBlocks)
+	,	BlockSize(InFFTAlgorithm->Size() / 2)
+	,	FFTAlgorithm(InFFTAlgorithm)
+	,	NumFFTOutputFloats(InFFTAlgorithm->NumOutputFloats())
+	,	NumActiveBlocks(0)
+	,	NumImpulseResponseSamples(0)
+	{
+		check(FFTAlgorithm->Size() > 0);
+		check(NumFFTOutputFloats > 0);
+		check(BlockSize > 0);
+
+		FFTInput.AddZeroed(FFTAlgorithm->Size());
+
+		Blocks.AddDefaulted(NumBlocks);
+
+		for (int32 i = 0; i < NumBlocks; i++)
+		{
+			Blocks[i].AddZeroed(NumFFTOutputFloats);
+		}
+	}
+
+	FUniformPartitionConvolution::FImpulseResponse::~FImpulseResponse()
+	{
+	}
+
+	void FUniformPartitionConvolution::FImpulseResponse::SetImpulseResponse(const float* InSamples, int32 InNum)
+	{
+		check(InNum >= 0);
+
+		// Check that input length is within limits
+		const int32 MaxNumSamples = NumBlocks * BlockSize;
+		if (InNum > MaxNumSamples)
+		{
+			UE_LOG(LogSignalProcessing, Warning, TEXT("Truncating impulse response of %d samples to maximum length of %d samples"), InNum, MaxNumSamples);
+			InNum = MaxNumSamples;
+		}
+		NumImpulseResponseSamples = InNum;
+
+		// Zero out previous blocks
+		for (int32 i = 0; i < NumBlocks; i++)
+		{
+			FMemory::Memset(Blocks[i].GetData(), 0, sizeof(float) * NumFFTOutputFloats);
+		}
+
+		NumActiveBlocks = InNum / BlockSize;
+		// Add extra block to account for round-down nature of integer division.
+		NumActiveBlocks += InNum % BlockSize == 0 ? 0 : 1;
+
+		// Store partitioned & transformed version of impulse response. 
+		for (int32 i = 0; i < NumActiveBlocks; i++)
+		{
+			// Copy samples to input buffer.
+			const int32 ReadPos = i * BlockSize;
+			const int32 NumSamplesToCopy = FMath::Min(BlockSize, NumImpulseResponseSamples - ReadPos);
+			const int32 NumSamplesToZero = BlockSize - NumSamplesToCopy;
+
+			FMemory::Memcpy(FFTInput.GetData(), &InSamples[ReadPos], NumSamplesToCopy * sizeof(float));
+			if (NumSamplesToZero > 0)
+			{
+				FMemory::Memset(&FFTInput.GetData()[NumSamplesToCopy], 0, NumSamplesToZero * sizeof(float));
+			}
+
+			// Perform fft
+			FFTAlgorithm->ForwardRealToComplex(FFTInput.GetData(), Blocks[i].GetData());	
+		}
+	}
+
+	const AlignedFloatBuffer& FUniformPartitionConvolution::FImpulseResponse::GetTransformedBlock(int32 InBlockIndex) const
+	{
+		return Blocks[InBlockIndex];
+	}
+
+	int32 FUniformPartitionConvolution::FImpulseResponse::GetNumActiveBlocks() const
+	{
+		return NumActiveBlocks;
+	}
+
+	int32 FUniformPartitionConvolution::FImpulseResponse::GetNumImpulseResponseSamples() const
+	{
+		return NumImpulseResponseSamples;
+	}
+
+
+	FUniformPartitionConvolutionFactory::~FUniformPartitionConvolutionFactory()
+	{
+	}
+
+	/** Name of this particular factory. */
+	const FName FUniformPartitionConvolutionFactory::GetFactoryName() const
+	{
+		static const FName FactoryName(TEXT("UniformPartitionConvolutionFactory"));
+		return FactoryName;
+	}
+
+	/** If true, this implementation uses hardware acceleration. */
+	bool FUniformPartitionConvolutionFactory::IsHardwareAccelerated() const
+	{
+		return false;
+	}
+
+	/** Returns true if the input settings are supported by this factory. */
+	bool FUniformPartitionConvolutionFactory::AreConvolutionSettingsSupported(const FConvolutionSettings& InSettings) const
+	{
+		bool bIsValidNumInputs = InSettings.NumInputChannels > 0;
+		bool bIsValidNumOutputs = InSettings.NumOutputChannels > 0;
+		bool bIsValidNumImpulseResponses = InSettings.NumImpulseResponses > 0;
+		bool bIsValidMaxNumImpulseResponseSamples = InSettings.MaxNumImpulseResponseSamples > 0;
+		// Block Size must be greater than zero and a power of 2.
+		bool bIsValidBlockSize = (InSettings.BlockNumSamples > 0) && (FMath::CountBits(InSettings.BlockNumSamples) == 1);
+
+		// Check if FFFTFactory supports required fftsettings
+		FFFTSettings FFTSettings = FUniformPartitionConvolutionFactory::FFTSettingsFromConvolutionSettings(InSettings);
+		bool bSupportedFFTSettings = FFFTFactory::AreFFTSettingsSupported(FFTSettings);
+
+		bool bIsSupported = bIsValidNumInputs;
+		bIsSupported &= bIsValidNumOutputs;
+		bIsSupported &= bIsValidNumImpulseResponses;
+		bIsSupported &= bIsValidMaxNumImpulseResponseSamples;
+		bIsSupported &= bIsValidBlockSize;
+		bIsSupported &= bSupportedFFTSettings;
+
+		return bIsSupported;
+	}
+
+	/** Creates a new Convolution algorithm. */
+	TUniquePtr<IConvolutionAlgorithm> FUniformPartitionConvolutionFactory::NewConvolutionAlgorithm(const FConvolutionSettings& InSettings)
+	{
+		check(AreConvolutionSettingsSupported(InSettings));
+
+		// Create required FFT settings based on convolution settings.
+		FFFTSettings FFTSettings = FUniformPartitionConvolutionFactory::FFTSettingsFromConvolutionSettings(InSettings);
+
+		// Create an FFT algorithm
+		TUniquePtr<IFFTAlgorithm> UniqueFFTAlgorithm = FFFTFactory::NewFFTAlgorithm(FFTSettings);
+		if (!UniqueFFTAlgorithm.IsValid())
+		{
+			UE_LOG(LogSignalProcessing, Error, TEXT("Failed to create uniform partition convolution algorithm due to failed FFT creation"));
+			return TUniquePtr<IConvolutionAlgorithm>();
+		}
+
+		// Convert FFT to shared ref 
+		TSharedRef<IFFTAlgorithm> FFTRef(UniqueFFTAlgorithm.Release());
+
+		// Create UniformPartitionConvolutionSettings
+		FUniformPartitionConvolutionSettings UPConvSettings;
+		UPConvSettings.NumInputChannels = InSettings.NumInputChannels;
+		UPConvSettings.NumOutputChannels = InSettings.NumOutputChannels;
+		UPConvSettings.NumImpulseResponses = InSettings.NumImpulseResponses;
+		UPConvSettings.MaxNumImpulseResponseSamples = InSettings.MaxNumImpulseResponseSamples;
+
+		return MakeUnique<FUniformPartitionConvolution>(UPConvSettings, FFTRef);
+	}
+
+	FFFTSettings FUniformPartitionConvolutionFactory::FFTSettingsFromConvolutionSettings(const FConvolutionSettings& InConvolutionSettings)
+	{
+		FFFTSettings FFTSettings;
+		FFTSettings.bEnableHardwareAcceleration = InConvolutionSettings.bEnableHardwareAcceleration;
+		FFTSettings.bArrays128BitAligned = true;
+
+		int32 Log2FFTSize = 0;
+
+		if (InConvolutionSettings.BlockNumSamples)
+		{
+			int32 Log2BlockSize = FMath::CountTrailingZeros(InConvolutionSettings.BlockNumSamples);
+			// For overlap save method of convolution, FFT is twice as big as audio block. 
+			Log2FFTSize = Log2BlockSize + 1;
+		}
+
+		FFTSettings.Log2Size = Log2FFTSize;
+
+		return FFTSettings;
+	}
+}
diff --git a/Engine/Source/Runtime/SignalProcessing/Private/UniformPartitionConvolution.h b/Engine/Source/Runtime/SignalProcessing/Private/UniformPartitionConvolution.h
new file mode 100644
--- /dev/null
+++ b/Engine/Source/Runtime/SignalProcessing/Private/UniformPartitionConvolution.h
@@ -0,0 +1,283 @@
+// Copyright Epic Games, Inc. All Rights Reserved.
+
+#pragma once
+
+#include "CoreMinimal.h"
+#include "DSP/BufferVectorOperations.h"
+#include "DSP/ConvolutionAlgorithm.h"
+#include "DSP/FFTAlgorithm.h"
+
+namespace Audio
+{
+	/** Settings for creating a uniform partition convolution. */
+	struct FUniformPartitionConvolutionSettings
+	{
+		/** Number of input channels for convolution. */
+		int32 NumInputChannels = 0;
+
+		/** Number of output channels for convolution. */
+		int32 NumOutputChannels = 0; 
+
+		/** Number of impulse responses for convolution. */
+		int32 NumImpulseResponses = 0; 
+
+		/** Maximum number of samples in an impulse response. */
+		int32 MaxNumImpulseResponseSamples = 0;
+	};
+
+	/** FUniformPartitionConvolution
+	 *
+	 * FUniformPartitionConvolution implements a fft based convolution algorithm supporting multiple 
+	 * inputs, impulse responses and outputs. The algorithm subdivides an impulse response into 
+	 * uniform blocks which results in several positive side effects:
+	 *
+	 * 	1. Latency is set to the block size which in many cases is less than the length of the impulse response.
+	 * 	2. Increased performance over brute time domain convolution.
+	 * 	3. Reuse of fourier transformed data when inputs are mapped to multiple outputs.
+	 *
+	 * The block size and latency is determined by the FFT algorithm passed to the constructor. 
+	 * 	BlockSize = FFTAlgorithm->Size() / 2
+	 * 	Latency = FFTAlgorithm->Size() / 2
+	 */
+	class FUniformPartitionConvolution : public IConvolutionAlgorithm
+	{
+		public:
+			/** Shared reference to FFT algorithm. */
+			using FSharedFFTRef = TSharedRef<IFFTAlgorithm>;
+
+			/** Create a Uniform Partition Convolution Algorithm
+			 *
+			 * @params InSettings - Settings for algorithm.
+			 * @param InFFTAlgorithm - FFTAlgorithm used to transform inputs, outputs and impulse responses.
+			 */
+			FUniformPartitionConvolution(const FUniformPartitionConvolutionSettings& InSettings, FSharedFFTRef InFFTAlgorithm);
+
+			/** Destructor. */
+			virtual ~FUniformPartitionConvolution();
+
+			/** Returns the number of samples in an audio block. */
+			virtual int32 GetNumSamplesInBlock() const override;
+
+			/** Returns number of audio inputs. */
+			virtual int32 GetNumAudioInputs() const override;
+
+			/** Returns number of audio outputs. */
+			virtual int32 GetNumAudioOutputs() const override;
+
+			/** Process one block of audio.
+			 *
+			 * InSamples is processed by the impulse responses. The output is placed in OutSamples.
+			 *
+			 * @params InSamples - A 2D array of input deinterleaved audio samples. InSamples[GetNumAudioInputs()][GetNumSamplesInBlock()]
+			 * @params OutSamples - A 2D array of output deinterleaved audio samples. OutSamples[GetNumAudioOutputs()][GetNumSamplesInBlock()]
+			 *
+			 */
+			virtual void ProcessAudioBlock(const float* const InSamples[], float* const OutSamples[]) override;
+
+			/** Reset internal history buffer for all audio inputs. */
+			virtual void ResetAudioHistory() override;
+
+			/** Maximum supported length of impulse response. */
+			virtual int32 GetMaxNumImpulseResponseSamples() const override;
+
+			/** Return the number of impulse responses. */
+			virtual int32 GetNumImpulseResponses() const override;
+
+			/** Return the number of samples in an impulse response. */
+			virtual int32 GetNumImpulseResponseSamples(int32 InImpulseResponseIndex) const override;
+
+			/** Set impulse response values. */
+			virtual void SetImpulseResponse(int32 InImpulseResponseIndex, const float* InSamples, int32 InNumSamples) override;
+
+			/** Sets the gain between an audio input channel, impulse response and audio output channel.
+			 *
+			 * ([audio inputs] * [impulse responses]) x [gain matrix] = [audio outputs]
+			 */
+			virtual void SetMatrixGain(int32 InAudioInputIndex, int32 InImpulseResponseIndex, int32 InAudioOutputIndex, float InGain) override;
+
+			/** Gets the gain between an audio input channel, impulse response and audio output channel.
+			 *
+			 * ([audio inputs] * [impulse responses]) x [gain matrix] = [audio outputs]
+			 */
+			virtual float GetMatrixGain(int32 InAudioInputIndex, int32 InImpulseResponseIndex, int32 InAudioOutputIndex) const override;
+
+		private:
+
+			// Bit mask to determine how many values to apply SIMD optimization. 
+			static const int32 NumSimdMask;
+
+			// Triple index into 3D mixing matrix.
+			struct FInputTransformOutputGroup : public TTuple<int32, int32, int32>
+			{
+				FInputTransformOutputGroup(int32 InInputIndex, int32 InTransformIndex, int32 InOutputIndex)
+				:	TTuple<int32, int32, int32>(InInputIndex, InTransformIndex, InOutputIndex)
+				{}
+
+				int32 GetInputIndex() const
+				{
+					return Get<0>();
+				}
+
+				int32 GetTransformIndex() const
+				{
+					return Get<1>();
+				}
+
+				int32 GetOutputIndex() const
+				{
+					return Get<2>();
+				}
+			};
+
+			// Vector optimized complex multiply add
+			void VectorComplexMultiplyAdd(const AlignedFloatBuffer& InA, const AlignedFloatBuffer& InB, AlignedFloatBuffer& Out) const;
+
+			// Vector optimized multiply
+			void VectorMultiplyByConstant(const AlignedFloatBuffer& InBuffer, float InConstant, AlignedFloatBuffer& OutBuffer) const;
+
+			// Contains a single input channel.
+			//
+			// The input channel holds the FFT of the latest input block.
+			struct FInput
+			{
+					// 
+					FInput(FSharedFFTRef InFFTAlgorithm);
+
+					// Add one block of audio samples.
+					void PushBlock(const float* InSamples);
+
+					// Reference to latest transformed block.
+					const AlignedFloatBuffer& GetTransformedBlock() const;
+
+					// Clears the transformed block
+					void Reset();
+
+					// Number of samples in a block
+					const int32 BlockSize;
+
+				private: 
+
+					FSharedFFTRef FFTAlgorithm;
+					AlignedFloatBuffer InputBuffer;
+					AlignedFloatBuffer OutputBuffer;
+			};
+
+			// Contains a single output channel
+			//
+			// The output channel holds a running history of fourier transformed output blocks. 
+			// When a block is popped, it converts it from the fourier domain to the time domain.
+			struct FOutput
+			{
+					// Construct an output
+					//
+					// @param InNumBlocks - Number of running history blocks
+					// @param InFFTAlgorithm - FFT algorithm for performing inverse FFT.
+					FOutput(int32 InNumBlocks, FSharedFFTRef InFFTAlgorithm);
+
+					~FOutput();
+
+					// Returns the history block at the given index. 
+					AlignedFloatBuffer& GetTransformedBlock(int32 InBlockIndex);	
+
+					// Pop the latest block without performing inverse FFT
+					void PopBlock();
+
+					// Pop the latest block and perform inverse fft.
+					// Out samples should hold InFFTAlgorithm->Size() / 2 floats. 
+					void PopBlock(float* OutSamples);
+
+					// Reset history of output blocks.
+					void Reset();
+
+					const int32 NumBlocks;
+					const int32 BlockSize;
+					const int32 NumFFTOutputFloats;
+				private:
+
+					FSharedFFTRef FFTAlgorithm;
+					AlignedFloatBuffer OutputBuffer;
+
+					int32 HeadBlockIndex;
+					TArray<AlignedFloatBuffer> Blocks;
+
+			};
+
+			// Contains a single impulse response 
+			//
+			// The impulse response is broken into blocks and transformed to the fourier domain.
+			struct FImpulseResponse
+			{
+					// Constructor
+					//
+					// @params InNumBlocks - Maximum number of FFT blocks
+					// @param InFFTAlgorithm - FFT algorithm to convert samples to fourier.
+					FImpulseResponse(int32 InNumBlocks, FSharedFFTRef InFFTAlgorithm);
+
+					~FImpulseResponse();
+
+					// Sets the impulse response.
+					void SetImpulseResponse(const float* InSamples, int32 InNum);
+
+					// Returns a block of the transformed impulse response.
+					const AlignedFloatBuffer& GetTransformedBlock(int32 InBlockIndex) const;
+
+					// Returns the number of active blocks. This may be less than the maximum
+					// number of blocks if the impulse response is shorter than the maximum
+					// number of impulse response samples. 
+					int32 GetNumActiveBlocks() const;
+
+					// Returns the number of samples in the impulse response.
+					int32 GetNumImpulseResponseSamples() const;
+
+					const int32 NumBlocks;
+					const int32 BlockSize;
+				private:
+
+					FSharedFFTRef FFTAlgorithm;
+					int32 NumFFTOutputFloats;		
+					int32 NumActiveBlocks;
+					int32 NumImpulseResponseSamples;
+					AlignedFloatBuffer FFTInput;
+					TArray<AlignedFloatBuffer> Blocks;
+			};
+
+			FUniformPartitionConvolutionSettings Settings;
+
+			int32 BlockSize;
+			int32 NumFFTOutputFloats;
+			int32 NumBlocks;
+			AlignedFloatBuffer TransformedAndScaledInput;
+			TSharedRef<IFFTAlgorithm> FFTAlgorithm;
+
+			TArray<FInput> Inputs;
+			TArray<FOutput> Outputs;
+			TArray<FImpulseResponse> ImpulseResponses;
+
+			TMap<FInputTransformOutputGroup, float> NonZeroGains;
+			TArray<bool> IsOutputZero;
+	};
+
+	/** FUniformPartitionConvolutionFactory creates FUniformPartitionConvolution algorithms. */
+	class FUniformPartitionConvolutionFactory : public IConvolutionAlgorithmFactory
+	{
+		public:
+			virtual ~FUniformPartitionConvolutionFactory();
+
+			/** Name of this particular factory. */
+			virtual const FName GetFactoryName() const override;
+
+			/** If true, this implementation uses hardware acceleration. */
+			virtual bool IsHardwareAccelerated() const override;
+
+			/** Returns true if the input settings are supported by this factory. */
+			virtual bool AreConvolutionSettingsSupported(const FConvolutionSettings& InSettings) const override;
+
+			/** Creates a new Convolution algorithm. */
+			virtual TUniquePtr<IConvolutionAlgorithm> NewConvolutionAlgorithm(const FConvolutionSettings& InSettings) override;
+
+		private:
+
+			// Convert the convolution settings into FFTAlgorithm settings.
+			static FFFTSettings FFTSettingsFromConvolutionSettings(const FConvolutionSettings& InConvolutionSettings);
+	};
+}
diff --git a/Engine/Source/Runtime/SignalProcessing/Private/VectorFFT.cpp b/Engine/Source/Runtime/SignalProcessing/Private/VectorFFT.cpp
new file mode 100644
--- /dev/null
+++ b/Engine/Source/Runtime/SignalProcessing/Private/VectorFFT.cpp
@@ -0,0 +1,1302 @@
+// Copyright Epic Games, Inc. All Rights Reserved.
+
+#include "VectorFFT.h"
+#include "SignalProcessingModule.h"
+#include "CoreMinimal.h"
+
+#include "DSP/FFTAlgorithm.h"
+#include "DSP/BufferVectorOperations.h"
+
+namespace Audio
+{
+	// Implementation of a complex FFT.
+	class FVectorComplexFFT
+	{
+		public:
+			// Minimum size of fft is required to support 2 radix-4 fft stages. 
+			static const int32 MinLog2FFTSize = 4;
+
+			// Maximum size of fft is set to avoid exceedingly large allocs and support
+			// various logic for indexing blocks internally.
+			static const int32 MaxLog2FFTSize = 16;
+
+			// Constructor
+			//
+			// @param InLog2FFTSize - Log2 size of the FFT.
+			FVectorComplexFFT(int32 InLog2FFTSize)
+			:	Log2FFTSize(InLog2FFTSize)
+			,	FFTSize(0)
+			,	NumFloats(0)
+			{
+				check(Log2FFTSize >= MinLog2FFTSize);
+				check(Log2FFTSize <= MaxLog2FFTSize);
+
+				FFTSize = 1 << Log2FFTSize;
+				NumFloats = 2 * FFTSize; // Takes 2 floats to represent a complex number
+
+				InverseWorkBuffer.AddUninitialized(NumFloats);
+
+				// Pregenerate weights needed to calculate this size of FFT.
+				GenerateRadix4Weights();
+				GenerateFinalIndices();
+				GenerateFinalWeights();
+			}
+
+			~FVectorComplexFFT()
+			{
+			}
+
+			// Perform forward complex FFT
+			//
+			// @param InComplex - Interleaved complex data with (2 * FFTSize) num floats.
+			// @param OutComplex - Interleaved complex data with (2 * FFTSize) num floats.
+			void ForwardComplexToComplex(const float* RESTRICT InComplex, float* RESTRICT OutComplex)
+			{
+				checkf(IsAligned<const float*>(InComplex, AUDIO_SIMD_FLOAT_ALIGNMENT), TEXT("Memory must be aligned.")); 
+				checkf(IsAligned<float*>(OutComplex, AUDIO_SIMD_FLOAT_ALIGNMENT), TEXT("Memory must be aligned.")); 
+
+				// To perform FFT, must complete Log2FFTSize stages.  Each radix pass performs 2^m stages
+				// where 2^m is the radix number. So a radix-4 stage is radix-2^m or radix-2^2. Hence radix
+				// 4 performs two stages. Radix-8 is Radix-2^3, so it performs 3 stages.
+				int32 CompletedStages = 0;
+
+				if (Log2FFTSize & 1)
+				{
+					// If we have an odd number of stages, start with a radix-8 to 
+					// perform first 3 stages. 
+					Radix8ButterflyConstantWeight(InComplex, OutComplex, Log2FFTSize);
+					CompletedStages = 3;
+				}
+				else
+				{
+					// Ifawe have an even number of stages, start with a radix-4 to 
+					// perform first 2 stages.
+					Radix4ButterflyConstantWeight(InComplex, OutComplex, Log2FFTSize);
+					CompletedStages = 2;
+				}
+
+				// Fit in a few more constant weight radix4s if possible. This routine is faster
+				// than the default Radix4Butterfly because it does not need weights.
+				for (int32 StageIndex = CompletedStages; StageIndex < Log2FFTSize - 4 ; StageIndex += 2 )
+				{
+					Radix4ButterflyConstantWeight(OutComplex, OutComplex, Log2FFTSize - StageIndex);
+				}
+
+				// Perform a bunch of radix4 ffts with varying weights
+				// Logically you would arrange this loop to first iterate over stage indices, and then over
+				// butterfly indices, but by reorganizing the loop the code is more cache coherent.
+				for (int32 ButterflyIndex = 1 ; CompletedStages < (Log2FFTSize - 4) ; CompletedStages += 2)
+				{
+					for ( ; ButterflyIndex < (1 << CompletedStages); ++ButterflyIndex )
+					{
+						for (int32 StageIndex = CompletedStages; StageIndex < Log2FFTSize - 4 ; StageIndex += 2 )
+						{
+							Radix4Butterfly(OutComplex, ButterflyIndex, Log2FFTSize - StageIndex, Radix4Weights[ButterflyIndex]);
+						}
+					}
+				}
+
+				if (CompletedStages < (Log2FFTSize - 2))
+				{
+					// Special case for 2nd to last stage of fft which has better cache coherency
+					Radix4Butterfly2ndToFinal(OutComplex, Log2FFTSize - 4);
+				}
+
+				// Special case for last stage of fft with cache coherency tricks and index reversal built in.
+				Radix4ButterflyFinal(OutComplex, Log2FFTSize - 2);
+			}
+
+			void InverseComplexToComplex(const float* RESTRICT InComplex, float* RESTRICT OutComplex)
+			{
+				checkf(IsAligned<const float*>(InComplex, AUDIO_SIMD_FLOAT_ALIGNMENT), TEXT("Memory must be aligned.")); 
+				checkf(IsAligned<float*>(OutComplex, AUDIO_SIMD_FLOAT_ALIGNMENT), TEXT("Memory must be aligned.")); 
+
+				// Perform inverse FFT by complex conjugating the input and output.
+				
+				float* WorkData = InverseWorkBuffer.GetData();
+
+				ScaledComplexConjugate(InComplex, 1.f, WorkData, NumFloats);
+
+				ForwardComplexToComplex(WorkData, OutComplex);
+
+				const float Scale = 1.f / static_cast<float>(FFTSize);
+
+				ScaledComplexConjugate(OutComplex, Scale, OutComplex, NumFloats);
+			}
+
+		private:
+
+			// Weight structure for general radix 4 pass.
+			struct FRadix4Weight 
+			{
+				alignas(AUDIO_SIMD_BYTE_ALIGNMENT) float W1R[4]; 
+				alignas(AUDIO_SIMD_BYTE_ALIGNMENT) float W1RNeg[4]; // Negative version of W1R.
+
+				// These have specialized sign flips for calculations of D2 and  D3
+				alignas(AUDIO_SIMD_BYTE_ALIGNMENT) float W1RD2[4]; 
+				alignas(AUDIO_SIMD_BYTE_ALIGNMENT) float W1RD3[4]; 
+
+				alignas(AUDIO_SIMD_BYTE_ALIGNMENT) float W1I[4]; 
+				alignas(AUDIO_SIMD_BYTE_ALIGNMENT) float W2R[4];
+				alignas(AUDIO_SIMD_BYTE_ALIGNMENT) float W2RNeg[4]; // Negative Version of W2R
+				alignas(AUDIO_SIMD_BYTE_ALIGNMENT) float W2I[4]; 
+				alignas(AUDIO_SIMD_BYTE_ALIGNMENT) float W3R[4]; 
+				alignas(AUDIO_SIMD_BYTE_ALIGNMENT) float W3RNeg[4]; // Negative Version of W4R
+				alignas(AUDIO_SIMD_BYTE_ALIGNMENT) float W3I[4];
+			}; 
+
+
+			// Structure to hold loaded inputs on final radix 4 pass
+			struct FFinalInputs
+			{
+				VectorRegister A0;
+				VectorRegister A1;
+				VectorRegister A2;
+				VectorRegister A3;
+				VectorRegister A4;
+				VectorRegister A5;
+				VectorRegister A6;
+				VectorRegister A7;
+			};
+
+			// Structure to hold loaded outputs on final radix 4 pass
+			struct FFinalOutputs
+			{
+				VectorRegister D0;
+				VectorRegister D1;
+				VectorRegister D2;
+				VectorRegister D3;
+				VectorRegister D4;
+				VectorRegister D5;
+				VectorRegister D6;
+				VectorRegister D7;
+			};
+
+			// Weight structure for final Radix4 pass
+			struct FFinalWeights
+			{
+				alignas(AUDIO_SIMD_BYTE_ALIGNMENT) float W1R[4]; 
+				alignas(AUDIO_SIMD_BYTE_ALIGNMENT) float W1I[4]; 
+				alignas(AUDIO_SIMD_BYTE_ALIGNMENT) float W2R[4];
+				alignas(AUDIO_SIMD_BYTE_ALIGNMENT) float W2I[4]; 
+				alignas(AUDIO_SIMD_BYTE_ALIGNMENT) float W3R[4]; 
+				alignas(AUDIO_SIMD_BYTE_ALIGNMENT) float W3I[4];
+				alignas(AUDIO_SIMD_BYTE_ALIGNMENT) float W4R[4]; 
+				alignas(AUDIO_SIMD_BYTE_ALIGNMENT) float W4I[4];
+				alignas(AUDIO_SIMD_BYTE_ALIGNMENT) float W5R[4]; 
+				alignas(AUDIO_SIMD_BYTE_ALIGNMENT) float W5I[4];
+				alignas(AUDIO_SIMD_BYTE_ALIGNMENT) float W6R[4]; 
+				alignas(AUDIO_SIMD_BYTE_ALIGNMENT) float W6I[4];
+				alignas(AUDIO_SIMD_BYTE_ALIGNMENT) float W7R[4]; 
+				alignas(AUDIO_SIMD_BYTE_ALIGNMENT) float W7I[4];
+
+				alignas(AUDIO_SIMD_BYTE_ALIGNMENT) float W2RNeg[4]; // Negative version of W2R
+				alignas(AUDIO_SIMD_BYTE_ALIGNMENT) float W2RD4[4]; // Special case for calculating D4
+				alignas(AUDIO_SIMD_BYTE_ALIGNMENT) float W2RD6[4]; // Special case for calculating D5
+				alignas(AUDIO_SIMD_BYTE_ALIGNMENT) float W3RD5[4];  // Special case for calculating D6
+				alignas(AUDIO_SIMD_BYTE_ALIGNMENT) float W3RD7[4];  // Special case for calculating D7
+				alignas(AUDIO_SIMD_BYTE_ALIGNMENT) float W3RNeg[4];  // Negative version of W3R
+				alignas(AUDIO_SIMD_BYTE_ALIGNMENT) float W4RNeg[4];  // Negative version of W4R
+				alignas(AUDIO_SIMD_BYTE_ALIGNMENT) float W5RNeg[4];  // Negative version of W5R
+				alignas(AUDIO_SIMD_BYTE_ALIGNMENT) float W6RNeg[4];  // Negative version of W6R
+				alignas(AUDIO_SIMD_BYTE_ALIGNMENT) float W7RNeg[4];  // Negative version of W7R
+			};
+
+			// Radix4 index locations for performing final Radix4 passes and bit reversal
+			// without overwriting needed data.
+			struct FFinalIndices
+			{
+				int32 ReadIndex;
+				int32 WriteIndex;
+			};
+
+
+			// Perform complex conjugate as well as scale.
+			//
+			// @param InValues - Array of floats representing complex values in interleave format.
+			// @param Scale - Scale to apply.
+			// @param OutValues - Array of floats representing complex values in interleave format.
+			// @param Num - Number of floats in array (NOT number of complex values).
+			void ScaledComplexConjugate(const float* InValues, float Scale, float* OutValues, int32 Num)
+			{
+				// Use mask to quickly find out number of values that can be SIMD'd
+				const int32 SIMD_MASK = 0xFFFFFFFC;
+				const int32 NumToSimd = SIMD_MASK & Num;
+
+				// Complex values in a vector are [real_1, imag_1, real_2, imag_2].  
+				// By multipling this value, we flip the sign of the imaginary components, which
+				// is the equivalent of a complex conjugate.
+				const VectorRegister SignFlipImag = MakeVectorRegister(Scale, -Scale, Scale, -Scale);
+
+				// Perform operation using SIMD
+				for (int32 i = 0; i < NumToSimd; i += 4)
+				{
+					VectorRegister Value = VectorLoadAligned(&InValues[i]);
+					Value = VectorMultiply(SignFlipImag, Value);
+					VectorStoreAligned(Value, &OutValues[i]);
+				}
+
+				// Perform operation where SIMD not possible.
+				for (int32 i = NumToSimd; i < Num; i+= 2)
+				{
+					OutValues[i] = Scale * InValues[i];
+					OutValues[i + 1] = -Scale * InValues[i + 1];
+				}
+			}
+
+			// Perform a radix-4 butterfly which uses constant weights.
+			void Radix4ButterflyConstantWeight(const float* InValues, float* OutValues, int32 InStageIndex)
+			{
+				// This routine only supported when stage index is greater than or equal to 3. 
+				// This comes about for two reasons, 
+				// 	1. A Radix-4 processes two stages. (m = StageCount, 2^m = 4).
+				// 	2. SIMD optimizations processes two butterflies in parallel, so must have an
+				// 	   even number of butterflies. This sets the mimum stage index to 3, which will
+				// 	   result in minimally 2 radix-4 butterflies being calculated. 
+				check(InStageIndex >= 3);
+
+				// Calculate number of constant weight butterflies in this stage.
+				const int NumButterflies = 1 << (InStageIndex - 2);
+
+				const VectorRegister SignFlipImag = MakeVectorRegister(1.f, -1.f, 1.f, -1.f);
+
+				const int32 Offset0 = 0;
+				const int32 Offset1 = 2 * NumButterflies;
+				const int32 Offset2 = 4 * NumButterflies;
+				const int32 Offset3 = 6 * NumButterflies;
+
+				for (int32 i = 0; i < NumButterflies; i += 2)
+				{
+					const int32 Pos = 2 * i;
+					const int32 Pos0 = Offset0 + Pos;
+					const int32 Pos1 = Offset1 + Pos;
+					const int32 Pos2 = Offset2 + Pos;
+					const int32 Pos3 = Offset3 + Pos;
+
+					VectorRegister A0 = VectorLoadAligned(&InValues[Pos0]);
+					VectorRegister A1 = VectorLoadAligned(&InValues[Pos1]);
+					VectorRegister A2 = VectorLoadAligned(&InValues[Pos2]);
+					VectorRegister A3 = VectorLoadAligned(&InValues[Pos3]);
+
+					VectorRegister C0 = VectorAdd(A0, A2);
+					VectorRegister C2 = VectorSubtract(A0, A2);
+					VectorRegister C1 = VectorAdd(A1, A3);
+					VectorRegister C3 = VectorSubtract(A1, A3);
+
+					VectorRegister C3Conj = VectorMultiply(C3, SignFlipImag);
+					VectorRegister C3ConjSwizzle = VectorSwizzle(C3Conj, 1, 0, 3, 2);
+
+					VectorRegister D0 = VectorAdd(C1, C0);
+					VectorRegister D1 = VectorSubtract(C0, C1);
+					VectorRegister D2 = VectorAdd(C2, C3ConjSwizzle);
+					VectorRegister D3 = VectorSubtract(C2, C3ConjSwizzle);
+
+					VectorStoreAligned(D0, &OutValues[Pos0]);
+					VectorStoreAligned(D1, &OutValues[Pos1]);
+					VectorStoreAligned(D2, &OutValues[Pos2]);
+					VectorStoreAligned(D3, &OutValues[Pos3]);
+				}
+			}
+
+			// Perform a radix-8 butterfly which uses constant weights.
+			void Radix8ButterflyConstantWeight(const float* InValues, float* OutValues, int32 InStageIndex)
+			{
+
+				const float Sqrt2D2 = .7071067811865475244f;
+
+				// This routine only supported when stage index is greater than or equal to four. 
+				// This comes about for two reasons, 
+				// 	1. A Radix-8 processes three stages. (m = StageCount, 2^m = 8).
+				// 	2. SIMD optimizations processes two butterflies in parallel, so must have an
+				// 	   even number of butterflies. This sets the mimum stage index to 4, which will
+				// 	   result in minimally 2 radix-8 butterflies being calculated. 
+				check(InStageIndex >= 4);
+
+				// Calculate number of constant weight butterflies in this stage.
+				const int32 NumButterflies = 1 << (InStageIndex - 3);
+
+				const VectorRegister SignFlipImag = MakeVectorRegister(1.f, -1.f, 1.f, -1.f);
+				const VectorRegister VectorSqrt2D2 = MakeVectorRegister(Sqrt2D2, Sqrt2D2, Sqrt2D2, Sqrt2D2);
+				const VectorRegister VectorNegSqrt2D2 = MakeVectorRegister(-Sqrt2D2, -Sqrt2D2, -Sqrt2D2, -Sqrt2D2);
+
+				const int32 Offset0 = 0;
+				const int32 Offset1 = 2 * NumButterflies;
+				const int32 Offset2 = 4 * NumButterflies;
+				const int32 Offset3 = 6 * NumButterflies;
+				const int32 Offset4 = 8 * NumButterflies;
+				const int32 Offset5 = 10 * NumButterflies;
+				const int32 Offset6 = 12 * NumButterflies;
+				const int32 Offset7 = 14 * NumButterflies;
+
+				for (int32 i = 0; i < NumButterflies; i += 2)
+				{
+					const int32 Pos = 2 * i;
+					const int32 Pos0 = Pos;
+					const int32 Pos1 = Offset1 + Pos;
+					const int32 Pos2 = Offset2 + Pos;
+					const int32 Pos3 = Offset3 + Pos;
+					const int32 Pos4 = Offset4 + Pos;
+					const int32 Pos5 = Offset5 + Pos;
+					const int32 Pos6 = Offset6 + Pos;
+					const int32 Pos7 = Offset7 + Pos;
+
+					VectorRegister A0 = VectorLoadAligned(&InValues[Pos0]);
+					VectorRegister A1 = VectorLoadAligned(&InValues[Pos1]);
+					VectorRegister A2 = VectorLoadAligned(&InValues[Pos2]);
+					VectorRegister A3 = VectorLoadAligned(&InValues[Pos3]);
+					VectorRegister A4 = VectorLoadAligned(&InValues[Pos4]);
+					VectorRegister A5 = VectorLoadAligned(&InValues[Pos5]);
+					VectorRegister A6 = VectorLoadAligned(&InValues[Pos6]);
+					VectorRegister A7 = VectorLoadAligned(&InValues[Pos7]);
+
+					VectorRegister B0 = VectorAdd(A0, A4);
+					VectorRegister B1 = VectorAdd(A1, A5);
+					VectorRegister B2 = VectorAdd(A2, A6);
+					VectorRegister B3 = VectorAdd(A3, A7);
+					VectorRegister B4 = VectorSubtract(A0, A4);
+					VectorRegister B5 = VectorSubtract(A1, A5);
+					VectorRegister B6 = VectorSubtract(A2, A6);
+					VectorRegister B7 = VectorSubtract(A3, A7);
+
+					VectorRegister B6Conj = VectorMultiply(SignFlipImag, B6);
+					VectorRegister B6ConjSwizzle = VectorSwizzle(B6Conj, 1, 0, 3, 2);
+
+					VectorRegister B7Conj = VectorMultiply(SignFlipImag, B7);
+					VectorRegister B7ConjSwizzle = VectorSwizzle(B7Conj, 1, 0, 3, 2);
+
+					VectorRegister C0 = VectorAdd(B0, B2);
+					VectorRegister C1 = VectorAdd(B1, B3);
+					VectorRegister C2 = VectorSubtract(B0, B2);
+					VectorRegister C3 = VectorSubtract(B1, B3);
+					VectorRegister C4 = VectorAdd(B4, B6ConjSwizzle);
+					VectorRegister C5 = VectorAdd(B5, B7ConjSwizzle);
+					VectorRegister C6 = VectorSubtract(B4, B6ConjSwizzle);
+					VectorRegister C7 = VectorSubtract(B5, B7ConjSwizzle);
+
+					VectorRegister C3Conj = VectorMultiply(SignFlipImag, C3);
+					VectorRegister C3ConjSwizzle = VectorSwizzle(C3Conj, 1, 0, 3, 2);
+
+					VectorRegister C5Conj = VectorMultiply(SignFlipImag, C5);
+					VectorRegister C5ConjSwizzle = VectorSwizzle(C5Conj, 1, 0, 3, 2);
+					VectorRegister T5 = VectorAdd(C5, C5ConjSwizzle);
+
+					VectorRegister C7Conj = VectorMultiply(SignFlipImag, C7);
+					VectorRegister C7Swizzle = VectorSwizzle(C7, 1, 0, 3, 2);
+
+					VectorRegister T7 = VectorAdd(C7Swizzle, C7Conj);
+					VectorRegister T7Conj = VectorMultiply(T7, SignFlipImag);
+
+					VectorRegister D0 = VectorAdd(C0, C1);
+					VectorRegister D1 = VectorSubtract(C0, C1);
+					VectorRegister D2 = VectorAdd(C2, C3ConjSwizzle);
+					VectorRegister D3 = VectorSubtract(C2, C3ConjSwizzle);
+					VectorRegister D4 = VectorMultiplyAdd(T5, VectorSqrt2D2, C4);
+					VectorRegister D5 = VectorMultiplyAdd(T5, VectorNegSqrt2D2, C4);
+					VectorRegister D6 = VectorMultiplyAdd(VectorNegSqrt2D2, T7Conj, C6);
+					VectorRegister D7 = VectorMultiplyAdd(VectorSqrt2D2, T7Conj, C6);
+
+					VectorStoreAligned(D0, &OutValues[Pos0]);
+					VectorStoreAligned(D1, &OutValues[Pos1]);
+					VectorStoreAligned(D2, &OutValues[Pos2]);
+					VectorStoreAligned(D3, &OutValues[Pos3]);
+					VectorStoreAligned(D4, &OutValues[Pos4]);
+					VectorStoreAligned(D5, &OutValues[Pos5]);
+					VectorStoreAligned(D6, &OutValues[Pos6]);
+					VectorStoreAligned(D7, &OutValues[Pos7]);
+				}
+			}
+
+			// Perform a radix4 butterfly with dynamic weights.
+			void Radix4Butterfly(float* InOutValues, int32 ButterflyIndex, int32 InStageIndex, const FRadix4Weight& Weights)
+			{
+				// Number of values between butterflies.
+				const int32 Stride = 1 << InStageIndex;
+
+				const int32 NumButterflies = 1 << (InStageIndex - 2);
+
+				// Load weights for butterfly
+				const VectorRegister Weight1Real = VectorLoadAligned(Weights.W1R);
+				const VectorRegister Weight1Imag = VectorLoadAligned(Weights.W1I);
+				const VectorRegister Weight2Real = VectorLoadAligned(Weights.W2R);
+				const VectorRegister Weight2Imag = VectorLoadAligned(Weights.W2I);
+				const VectorRegister Weight3Real = VectorLoadAligned(Weights.W3R);
+				const VectorRegister Weight3Imag = VectorLoadAligned(Weights.W3I);
+
+				const VectorRegister Weight3RealNeg = VectorLoadAligned(Weights.W3RNeg);
+				const VectorRegister Weight2RealNeg = VectorLoadAligned(Weights.W2RNeg);
+				const VectorRegister Weight1RealNeg = VectorLoadAligned(Weights.W1RNeg);
+				const VectorRegister Weight1RealD2 = VectorLoadAligned(Weights.W1RD2);
+				const VectorRegister Weight1RealD3 = VectorLoadAligned(Weights.W1RD3);
+
+				// Perform butterflies.
+				for (int32 i = 0; i < NumButterflies; i += 2)
+				{
+					const int32 Pos0 = 2 * (Stride * ButterflyIndex + i);
+					const int32 Pos1 = 2 * (Stride * ButterflyIndex + 1 * NumButterflies + i);
+					const int32 Pos2 = 2 * (Stride * ButterflyIndex + 2 * NumButterflies + i);
+					const int32 Pos3 = 2 * (Stride * ButterflyIndex + 3 * NumButterflies + i);
+
+					VectorRegister A0 = VectorLoadAligned(&InOutValues[Pos0]);
+					VectorRegister A1 = VectorLoadAligned(&InOutValues[Pos1]);
+					VectorRegister A2 = VectorLoadAligned(&InOutValues[Pos2]);
+					VectorRegister A3 = VectorLoadAligned(&InOutValues[Pos3]);
+
+					VectorRegister A1Swizzle = VectorSwizzle(A1, 1, 0, 3, 2);
+					VectorRegister A2Swizzle = VectorSwizzle(A2, 1, 0, 3, 2);
+					VectorRegister A3Swizzle = VectorSwizzle(A3, 1, 0, 3, 2);
+
+					VectorRegister B1 = VectorMultiplyAdd(A1Swizzle, Weight1Imag, A1);
+					VectorRegister B2 = VectorMultiplyAdd(A2Swizzle, Weight2Imag, A2);
+					VectorRegister B3 = VectorMultiplyAdd(A3Swizzle, Weight3Imag, A3);
+
+					VectorRegister C0 = VectorMultiplyAdd(B2, Weight2Real, A0);
+					VectorRegister C2 = VectorMultiplyAdd(B2, Weight2RealNeg, A0);
+					VectorRegister C1 = VectorMultiplyAdd(B3, Weight3Real, B1);
+					VectorRegister C3 = VectorMultiplyAdd(B3, Weight3RealNeg, B1);
+
+					VectorRegister C3Swizzle = VectorSwizzle(C3, 1, 0, 3, 2);
+
+					VectorRegister D0 = VectorMultiplyAdd(C1, Weight1Real, C0);
+					VectorRegister D1 = VectorMultiplyAdd(C1, Weight1RealNeg, C0);
+					VectorRegister D2 = VectorMultiplyAdd(C3Swizzle, Weight1RealD2, C2);
+					VectorRegister D3 = VectorMultiplyAdd(C3Swizzle, Weight1RealD3, C2);
+
+					VectorStoreAligned(D0, &InOutValues[Pos0]);
+					VectorStoreAligned(D1, &InOutValues[Pos1]);
+					VectorStoreAligned(D2, &InOutValues[Pos2]);
+					VectorStoreAligned(D3, &InOutValues[Pos3]);
+				}
+			}
+
+			// Special case of 2nd to last radix4 which has to load new weights for
+			// each iteration.
+			void Radix4Butterfly2ndToFinal(float* InOutValues, int32 StageIndex)
+			{
+				int32 NumPasses = 1 << StageIndex;
+
+				// Elements between passes.
+				const int32 Stride = 16;
+
+				for (int32 i = 0; i < NumPasses; ++i)
+				{
+					// Load values for current weight.
+					const FRadix4Weight& Weights = Radix4Weights[i];
+
+					const VectorRegister Weight1Real = VectorLoadAligned(Weights.W1R);
+					const VectorRegister Weight1Imag = VectorLoadAligned(Weights.W1I);
+					const VectorRegister Weight2Real = VectorLoadAligned(Weights.W2R);
+					const VectorRegister Weight2Imag = VectorLoadAligned(Weights.W2I);
+					const VectorRegister Weight3Real = VectorLoadAligned(Weights.W3R);
+					const VectorRegister Weight3Imag = VectorLoadAligned(Weights.W3I);
+
+					const VectorRegister Weight3RealNeg = VectorLoadAligned(Weights.W3RNeg);
+					const VectorRegister Weight2RealNeg = VectorLoadAligned(Weights.W2RNeg);
+					const VectorRegister Weight1RealNeg = VectorLoadAligned(Weights.W1RNeg);
+					const VectorRegister Weight1RealD2 = VectorLoadAligned(Weights.W1RD2);
+					const VectorRegister Weight1RealD3 = VectorLoadAligned(Weights.W1RD3);
+
+					for (int32 j = 0; j < 4 ; j += 2)
+					{
+						const int32 Pos0 = 2 * (Stride * i + j);
+						const int32 Pos1 = 2 * (Stride * i + 4 + j);
+						const int32 Pos2 = 2 * (Stride * i + 8 + j);
+						const int32 Pos3 = 2 * (Stride * i + 12 + j);
+
+						VectorRegister A0 = VectorLoadAligned(&InOutValues[Pos0]);
+						VectorRegister A1 = VectorLoadAligned(&InOutValues[Pos1]);
+						VectorRegister A2 = VectorLoadAligned(&InOutValues[Pos2]);
+						VectorRegister A3 = VectorLoadAligned(&InOutValues[Pos3]);
+
+						VectorRegister A1Swizzle = VectorSwizzle(A1, 1, 0, 3, 2);
+						VectorRegister A2Swizzle = VectorSwizzle(A2, 1, 0, 3, 2);
+						VectorRegister A3Swizzle = VectorSwizzle(A3, 1, 0, 3, 2);
+
+						VectorRegister B1 = VectorMultiplyAdd(A1Swizzle, Weight1Imag, A1);
+						VectorRegister B2 = VectorMultiplyAdd(A2Swizzle, Weight2Imag, A2);
+						VectorRegister B3 = VectorMultiplyAdd(A3Swizzle, Weight3Imag, A3);
+						
+						VectorRegister C0 = VectorMultiplyAdd(B2, Weight2Real, A0);
+						VectorRegister C1 = VectorMultiplyAdd(B3, Weight3Real, B1);
+						VectorRegister C2 = VectorMultiplyAdd(B2, Weight2RealNeg, A0);
+						VectorRegister C3 = VectorMultiplyAdd(B3, Weight3RealNeg, B1);
+						VectorRegister C3Swizzle = VectorSwizzle(C3, 1, 0, 3, 2);
+						
+						VectorRegister D0 = VectorMultiplyAdd(C1, Weight1Real, C0);
+						VectorRegister D1 = VectorMultiplyAdd(C1, Weight1RealNeg, C0);
+						VectorRegister D2 = VectorMultiplyAdd(C3Swizzle, Weight1RealD2, C2);
+						VectorRegister D3 = VectorMultiplyAdd(C3Swizzle, Weight1RealD3, C2);
+
+						VectorStoreAligned(D0, &InOutValues[Pos0]);
+						VectorStoreAligned(D1, &InOutValues[Pos1]);
+						VectorStoreAligned(D2, &InOutValues[Pos2]);
+						VectorStoreAligned(D3, &InOutValues[Pos3]);
+					}
+				}
+			}
+
+
+			// Read data for final butterfly.  Performs part of bit reversal order.
+			void ReadFinalButterflyInputs(const float* InValues, int32 InNumButterflies, int32 InReadIndex, FFinalInputs& OutValues)
+			{
+				const int32 Pos0 = 2 * (0 * InNumButterflies + 4 * InReadIndex);
+				const int32 Pos1 = Pos0 + 4;
+				const int32 Pos2 = 2 * (2 * InNumButterflies + 4 * InReadIndex);
+				const int32 Pos3 = Pos2 + 4;
+				const int32 Pos4 = 2 * (1 * InNumButterflies + 4 * InReadIndex);
+				const int32 Pos5 = Pos4 + 4;
+				const int32 Pos6 = 2 * (3 * InNumButterflies + 4 * InReadIndex);
+				const int32 Pos7 = Pos6 + 4;
+				
+				VectorRegister T0 = VectorLoadAligned(&InValues[Pos0]);
+				VectorRegister T1 = VectorLoadAligned(&InValues[Pos1]);
+				VectorRegister T2 = VectorLoadAligned(&InValues[Pos2]);
+				VectorRegister T3 = VectorLoadAligned(&InValues[Pos3]);
+				VectorRegister T4 = VectorLoadAligned(&InValues[Pos4]);
+				VectorRegister T5 = VectorLoadAligned(&InValues[Pos5]);
+				VectorRegister T6 = VectorLoadAligned(&InValues[Pos6]);
+				VectorRegister T7 = VectorLoadAligned(&InValues[Pos7]);
+
+				OutValues.A0 = VectorShuffle(T0, T2, 0, 1, 0, 1);
+				OutValues.A1 = VectorShuffle(T4, T6, 0, 1, 0, 1);
+				OutValues.A2 = VectorShuffle(T0, T2, 2, 3, 2, 3);
+				OutValues.A3 = VectorShuffle(T4, T6, 2, 3, 2, 3);
+				OutValues.A4 = VectorShuffle(T1, T3, 0, 1, 0, 1);
+				OutValues.A5 = VectorShuffle(T5, T7, 0, 1, 0, 1);
+				OutValues.A6 = VectorShuffle(T1, T3, 2, 3, 2, 3);
+				OutValues.A7 = VectorShuffle(T5, T7, 2, 3, 2, 3);
+			}
+
+			// Write data for final butterfly.  Performs part of bit reversal order.
+			void WriteFinalButterflyOutputs(const FFinalOutputs& InResult, int32 InNumButterflies, int32 InWriteIndex, float* OutValues)
+			{
+				const int32 Pos0 = 2 * (0 * InNumButterflies + 4 * InWriteIndex);
+				const int32 Pos1 = Pos0 + 4;
+				const int32 Pos2 = 2 * (2 * InNumButterflies + 4 * InWriteIndex);
+				const int32 Pos3 = Pos2 + 4;
+				const int32 Pos4 = 2 * (1 * InNumButterflies + 4 * InWriteIndex);
+				const int32 Pos5 = Pos4 + 4;
+				const int32 Pos6 = 2 * (3 * InNumButterflies + 4 * InWriteIndex);
+				const int32 Pos7 = Pos6 + 4;
+
+				VectorStoreAligned(InResult.D0, &OutValues[Pos0]);
+				VectorStoreAligned(InResult.D1, &OutValues[Pos1]);
+				VectorStoreAligned(InResult.D2, &OutValues[Pos2]);
+				VectorStoreAligned(InResult.D3, &OutValues[Pos3]);
+				VectorStoreAligned(InResult.D4, &OutValues[Pos4]);
+				VectorStoreAligned(InResult.D5, &OutValues[Pos5]);
+				VectorStoreAligned(InResult.D6, &OutValues[Pos6]);
+				VectorStoreAligned(InResult.D7, &OutValues[Pos7]);
+			}
+
+
+			// Compute butterfly in final stage.
+			void Radix4ButterflyFinalIteration(const FFinalInputs& Inputs, const FFinalWeights& InWeights, FFinalOutputs& Outputs)
+			{
+				// Note: Some weights are altered to bake in sign flips to avoid an extra multiply later on.
+				const VectorRegister W2I = VectorLoadAligned(InWeights.W2I);
+				const VectorRegister W2R = VectorLoadAligned(InWeights.W2R);
+				const VectorRegister W3I = VectorLoadAligned(InWeights.W3I);
+				const VectorRegister W3R = VectorLoadAligned(InWeights.W3R);
+				const VectorRegister W4I = VectorLoadAligned(InWeights.W4I);
+				const VectorRegister W4R = VectorLoadAligned(InWeights.W4R);
+				const VectorRegister W5I = VectorLoadAligned(InWeights.W5I);
+				const VectorRegister W5R = VectorLoadAligned(InWeights.W5R);
+				const VectorRegister W6I = VectorLoadAligned(InWeights.W6I);
+				const VectorRegister W6R = VectorLoadAligned(InWeights.W6R);
+				const VectorRegister W7I = VectorLoadAligned(InWeights.W7I);
+				const VectorRegister W7R = VectorLoadAligned(InWeights.W7R);
+
+				const VectorRegister W2RNeg = VectorLoadAligned(InWeights.W2RNeg);
+				const VectorRegister W3RNeg = VectorLoadAligned(InWeights.W3RNeg);
+				const VectorRegister W4RNeg = VectorLoadAligned(InWeights.W4RNeg);
+				const VectorRegister W5RNeg = VectorLoadAligned(InWeights.W5RNeg);
+				const VectorRegister W6RNeg = VectorLoadAligned(InWeights.W6RNeg);
+				const VectorRegister W7RNeg = VectorLoadAligned(InWeights.W7RNeg);
+
+				const VectorRegister W2RD4 = VectorLoadAligned(InWeights.W2RD4);
+				const VectorRegister W2RD6 = VectorLoadAligned(InWeights.W2RD6);
+				const VectorRegister W3RD5 = VectorLoadAligned(InWeights.W3RD5);
+				const VectorRegister W3RD7 = VectorLoadAligned(InWeights.W3RD7);
+
+				VectorRegister A2Swizzle = VectorSwizzle(Inputs.A2, 1, 0, 3, 2);
+				VectorRegister A3Swizzle = VectorSwizzle(Inputs.A3, 1, 0, 3, 2);
+				VectorRegister A4Swizzle = VectorSwizzle(Inputs.A4, 1, 0, 3, 2);
+				VectorRegister A5Swizzle = VectorSwizzle(Inputs.A5, 1, 0, 3, 2);
+				VectorRegister A6Swizzle = VectorSwizzle(Inputs.A6, 1, 0, 3, 2);
+				VectorRegister A7Swizzle = VectorSwizzle(Inputs.A7, 1, 0, 3, 2);
+
+				VectorRegister B2 = VectorMultiplyAdd(A2Swizzle, W2I, Inputs.A2);
+				VectorRegister B3 = VectorMultiplyAdd(A3Swizzle, W3I, Inputs.A3);
+				VectorRegister B4 = VectorMultiplyAdd(A4Swizzle, W4I, Inputs.A4);
+				VectorRegister B5 = VectorMultiplyAdd(A5Swizzle, W5I, Inputs.A5);
+				VectorRegister B6 = VectorMultiplyAdd(A6Swizzle, W6I, Inputs.A6);
+				VectorRegister B7 = VectorMultiplyAdd(A7Swizzle, W7I, Inputs.A7);
+
+				VectorRegister C0 = VectorMultiplyAdd(B4, W4R, Inputs.A0);
+				VectorRegister C1 = VectorMultiplyAdd(B5, W5R, Inputs.A1);
+				VectorRegister C2 = VectorMultiplyAdd(B6, W6R, B2);
+				VectorRegister C3 = VectorMultiplyAdd(B7, W7R, B3);
+				VectorRegister C4 = VectorMultiplyAdd(B4, W4RNeg, Inputs.A0);
+				VectorRegister C5 = VectorMultiplyAdd(B5, W5RNeg, Inputs.A1);
+				VectorRegister C6 = VectorMultiplyAdd(B6, W6RNeg, B2);
+				VectorRegister C7 = VectorMultiplyAdd(B7, W7RNeg, B3);
+
+				VectorRegister C6Swizzle = VectorSwizzle(C6, 1, 0, 3, 2);
+				VectorRegister C7Swizzle = VectorSwizzle(C7, 1, 0, 3, 2);
+
+				Outputs.D0 = VectorMultiplyAdd(C2, W2R, C0);
+				Outputs.D1 = VectorMultiplyAdd(C3, W3R, C1);
+				Outputs.D2 = VectorMultiplyAdd(C2, W2RNeg, C0);
+				Outputs.D3 = VectorMultiplyAdd(C3, W3RNeg, C1);
+				Outputs.D4 = VectorMultiplyAdd(C6Swizzle, W2RD4, C4);
+				Outputs.D5 = VectorMultiplyAdd(C7Swizzle, W3RD5, C5);
+				Outputs.D6 = VectorMultiplyAdd(C6Swizzle, W2RD6, C4);
+				Outputs.D7 = VectorMultiplyAdd(C7Swizzle, W3RD7, C5);
+			}
+
+
+			// Perform last set of radix 4 butterflies.
+			//
+			// This method is special since it also performs bit order reversal in a
+			// moderately cache coherent manner.
+			void Radix4ButterflyFinal(float* InOutValues, int InStageIndex)
+			{
+				int32 NumButterflies = 1 << InStageIndex;
+				int32 NumIterations = NumButterflies >> 2;
+
+				FFinalInputs Inputs;
+				FFinalOutputs Outputs;
+
+				int32 Iteration = 0;
+
+				ReadFinalButterflyInputs(InOutValues, NumButterflies, FinalIndices[Iteration].ReadIndex, Inputs);
+				Radix4ButterflyFinalIteration(Inputs, FinalWeights[Iteration], Outputs);
+
+				for (Iteration = 1; Iteration < NumIterations; ++Iteration)
+				{
+					ReadFinalButterflyInputs(InOutValues, NumButterflies, FinalIndices[Iteration].ReadIndex, Inputs);
+					WriteFinalButterflyOutputs(Outputs, NumButterflies, FinalIndices[Iteration - 1].WriteIndex, InOutValues);
+					Radix4ButterflyFinalIteration(Inputs, FinalWeights[Iteration], Outputs);
+				}
+
+				WriteFinalButterflyOutputs(Outputs, NumButterflies, FinalIndices[Iteration - 1].WriteIndex, InOutValues);
+			}
+
+			int32 IntLog2(int32 InValue)
+			{
+				check(InValue > 0);
+				check(FMath::CountBits(InValue) == 1);
+				return FMath::CountTrailingZeros(InValue);
+			}
+
+			void GenerateFinalIndices()
+			{
+				// Each pass of the final radix 4 operates on 16 complex values.  
+				const int32 FinalPassSize = 16;
+				const int32 NumFinalIndices = FFTSize / FinalPassSize;
+
+				// We need to ensure that FFTSize is at least 16 or else final
+				// FFT pass go past end of buffer.
+				check(FFTSize >= FinalPassSize);
+
+				FinalIndices.Reset();
+				FinalIndices.AddUninitialized(NumFinalIndices);
+
+				// These indices perform part of the bit order reversal on 16 element boundaries.
+				// Need to shift bits to get the bit reversed order on 16 element blocks.
+				const int32 Shift = 32 - (IntLog2(FFTSize) - 4);
+
+				int32 Index = 0;
+
+				for (int32 ReadIndex = 0; (ReadIndex < NumFinalIndices) && (Index < NumFinalIndices); ++ReadIndex)
+				{
+					// Get bit reversed order on 16 element block
+					const int32 WriteIndex = ReverseBits(ReadIndex) >> Shift;
+					
+					// If ReadIndex > WriteIndex, then ReadIndex in a previous iteration had the
+					// value WriteIndex has now, and we do not want to repeat it 
+					if (ReadIndex == WriteIndex)
+					{
+						// If equal, add one entry to read and write to same index.
+						FinalIndices[Index] = { ReadIndex, WriteIndex };
+						Index++;
+					}
+					else if (ReadIndex < WriteIndex)
+					{
+						// If ReadIndex < WriteIndex, add table entries in both orders.
+						// Loop logic in final radix pass will make sure that nothing
+						// gets overwritten.
+						FinalIndices[Index] = { ReadIndex, WriteIndex };
+						Index++;
+
+						FinalIndices[Index] = { WriteIndex, ReadIndex };
+						Index++;
+					}
+				}
+			}
+
+			void GenerateFinalWeights()
+			{
+				// Each pass of the final radix 4 operates on 16 complex values.  
+				const int32 FinalPassSize = 16;
+				const int32 NumFinalWeights = FFTSize / FinalPassSize;
+
+				// We need to ensure that FFTSize is at least 16 or else final
+				// FFT pass go past end of buffer.
+				check(FFTSize >= FinalPassSize);
+
+				FinalWeights.Reset();
+				FinalWeights.AddUninitialized(NumFinalWeights);
+
+				const double Scale = 1. / static_cast<double>(FFTSize);
+
+				for (int32 Pass = 0; Pass < NumFinalWeights; ++Pass)
+				{
+					const int ReadIndex = FinalIndices[Pass].ReadIndex;
+					const double RotatedBitFraction = RotateBitsAroundPoint(4 * ReadIndex);
+
+					const double Phase = 2. * PI * (RotatedBitFraction + 0 * Scale);
+					const double Phase1 = 2. * PI * (RotatedBitFraction + 1 * Scale);
+					const double Phase2 = 2. * PI * (RotatedBitFraction + 2 * Scale);
+					const double Phase3 = 2. * PI * (RotatedBitFraction + 3 * Scale);
+
+					FinalWeights[Pass].W2I[0] = -FMath::Tan(Phase);
+					FinalWeights[Pass].W2I[1] = FMath::Tan(Phase);
+					FinalWeights[Pass].W2I[2] = -FMath::Tan(Phase1);
+					FinalWeights[Pass].W2I[3] = FMath::Tan(Phase1);
+
+					FinalWeights[Pass].W3I[0] = -FMath::Tan(Phase2);
+					FinalWeights[Pass].W3I[1] = FMath::Tan(Phase2);
+					FinalWeights[Pass].W3I[2] = -FMath::Tan(Phase3);
+					FinalWeights[Pass].W3I[3] = FMath::Tan(Phase3);
+
+					FinalWeights[Pass].W2R[0] = FMath::Cos(Phase);
+					FinalWeights[Pass].W2R[1] = FMath::Cos(Phase);
+					FinalWeights[Pass].W2R[2] = FMath::Cos(Phase1);
+					FinalWeights[Pass].W2R[3] = FMath::Cos(Phase1);
+
+					FinalWeights[Pass].W2RD4[0] = -FMath::Cos(Phase);
+					FinalWeights[Pass].W2RD4[1] = FMath::Cos(Phase);
+					FinalWeights[Pass].W2RD4[2] = -FMath::Cos(Phase1);
+					FinalWeights[Pass].W2RD4[3] = FMath::Cos(Phase1);
+
+					FinalWeights[Pass].W2RD6[0] = FMath::Cos(Phase);
+					FinalWeights[Pass].W2RD6[1] = -FMath::Cos(Phase);
+					FinalWeights[Pass].W2RD6[2] = FMath::Cos(Phase1);
+					FinalWeights[Pass].W2RD6[3] = -FMath::Cos(Phase1);
+
+					FinalWeights[Pass].W2RNeg[0] = -FMath::Cos(Phase);
+					FinalWeights[Pass].W2RNeg[1] = -FMath::Cos(Phase);
+					FinalWeights[Pass].W2RNeg[2] = -FMath::Cos(Phase1);
+					FinalWeights[Pass].W2RNeg[3] = -FMath::Cos(Phase1);
+
+					FinalWeights[Pass].W3R[0] = FMath::Cos(Phase2);
+					FinalWeights[Pass].W3R[1] = FMath::Cos(Phase2);
+					FinalWeights[Pass].W3R[2] = FMath::Cos(Phase3);
+					FinalWeights[Pass].W3R[3] = FMath::Cos(Phase3);
+
+					FinalWeights[Pass].W3RD5[0] = -FMath::Cos(Phase2);
+					FinalWeights[Pass].W3RD5[1] = FMath::Cos(Phase2);
+					FinalWeights[Pass].W3RD5[2] = -FMath::Cos(Phase3);
+					FinalWeights[Pass].W3RD5[3] = FMath::Cos(Phase3);
+
+					FinalWeights[Pass].W3RD7[0] = FMath::Cos(Phase2);
+					FinalWeights[Pass].W3RD7[1] = -FMath::Cos(Phase2);
+					FinalWeights[Pass].W3RD7[2] = FMath::Cos(Phase3);
+					FinalWeights[Pass].W3RD7[3] = -FMath::Cos(Phase3);
+
+					FinalWeights[Pass].W3RNeg[0] = -FMath::Cos(Phase2);
+					FinalWeights[Pass].W3RNeg[1] = -FMath::Cos(Phase2);
+					FinalWeights[Pass].W3RNeg[2] = -FMath::Cos(Phase3);
+					FinalWeights[Pass].W3RNeg[3] = -FMath::Cos(Phase3);
+
+					FinalWeights[Pass].W4I[0] = -FMath::Tan(Phase + Phase);
+					FinalWeights[Pass].W4I[1] = FMath::Tan(Phase + Phase);
+					FinalWeights[Pass].W4I[2] = -FMath::Tan(Phase1 + Phase1);
+					FinalWeights[Pass].W4I[3] = FMath::Tan(Phase1 + Phase1);
+
+					FinalWeights[Pass].W5I[0] = -FMath::Tan(Phase2 + Phase2);
+					FinalWeights[Pass].W5I[1] = FMath::Tan(Phase2 + Phase2);
+					FinalWeights[Pass].W5I[2] = -FMath::Tan(Phase3 + Phase3);
+					FinalWeights[Pass].W5I[3] = FMath::Tan(Phase3 + Phase3);
+
+					FinalWeights[Pass].W4R[0] = FMath::Cos(Phase + Phase);
+					FinalWeights[Pass].W4R[1] = FMath::Cos(Phase + Phase);
+					FinalWeights[Pass].W4R[2] = FMath::Cos(Phase1 + Phase1);
+					FinalWeights[Pass].W4R[3] = FMath::Cos(Phase1 + Phase1);
+
+					FinalWeights[Pass].W4RNeg[0] = -FMath::Cos(Phase + Phase);
+					FinalWeights[Pass].W4RNeg[1] = -FMath::Cos(Phase + Phase);
+					FinalWeights[Pass].W4RNeg[2] = -FMath::Cos(Phase1 + Phase1);
+					FinalWeights[Pass].W4RNeg[3] = -FMath::Cos(Phase1 + Phase1);
+
+					FinalWeights[Pass].W5R[0] = FMath::Cos(Phase2 + Phase2);
+					FinalWeights[Pass].W5R[1] = FMath::Cos(Phase2 + Phase2);
+					FinalWeights[Pass].W5R[2] = FMath::Cos(Phase3 + Phase3);
+					FinalWeights[Pass].W5R[3] = FMath::Cos(Phase3 + Phase3);
+
+					FinalWeights[Pass].W5RNeg[0] = -FMath::Cos(Phase2 + Phase2);
+					FinalWeights[Pass].W5RNeg[1] = -FMath::Cos(Phase2 + Phase2);
+					FinalWeights[Pass].W5RNeg[2] = -FMath::Cos(Phase3 + Phase3);
+					FinalWeights[Pass].W5RNeg[3] = -FMath::Cos(Phase3 + Phase3);
+
+					FinalWeights[Pass].W6I[0] = -FMath::Tan(3. * Phase);
+					FinalWeights[Pass].W6I[1] = FMath::Tan(3. * Phase);
+					FinalWeights[Pass].W6I[2] = -FMath::Tan(3. * Phase1);
+					FinalWeights[Pass].W6I[3] = FMath::Tan(3. * Phase1);
+
+					FinalWeights[Pass].W7I[0] = -FMath::Tan(3. * Phase2);
+					FinalWeights[Pass].W7I[1] = FMath::Tan(3. * Phase2);
+					FinalWeights[Pass].W7I[2] = -FMath::Tan(3. * Phase3);
+					FinalWeights[Pass].W7I[3] = FMath::Tan(3. * Phase3);
+
+					FinalWeights[Pass].W6R[0] = 2 * FMath::Cos(Phase + Phase) - 1;
+					FinalWeights[Pass].W6R[1] = 2 * FMath::Cos(Phase + Phase) - 1;
+					FinalWeights[Pass].W6R[2] = 2 * FMath::Cos(Phase1 + Phase1) - 1;
+					FinalWeights[Pass].W6R[3] = 2 * FMath::Cos(Phase1 + Phase1) - 1; 
+
+					FinalWeights[Pass].W6RNeg[0] = -(2 * FMath::Cos(Phase + Phase) - 1);
+					FinalWeights[Pass].W6RNeg[1] = -(2 * FMath::Cos(Phase + Phase) - 1);
+					FinalWeights[Pass].W6RNeg[2] = -(2 * FMath::Cos(Phase1 + Phase1) - 1);
+					FinalWeights[Pass].W6RNeg[3] = -(2 * FMath::Cos(Phase1 + Phase1) - 1); 
+
+					FinalWeights[Pass].W7R[0] = 2 * FMath::Cos(Phase2 + Phase2) - 1;
+					FinalWeights[Pass].W7R[1] = 2 * FMath::Cos(Phase2 + Phase2) - 1;
+					FinalWeights[Pass].W7R[2] = 2 * FMath::Cos(Phase3 + Phase3) - 1;
+					FinalWeights[Pass].W7R[3] = 2 * FMath::Cos(Phase3 + Phase3) - 1;
+
+					FinalWeights[Pass].W7RNeg[0] = -(2 * FMath::Cos(Phase2 + Phase2) - 1);
+					FinalWeights[Pass].W7RNeg[1] = -(2 * FMath::Cos(Phase2 + Phase2) - 1);
+					FinalWeights[Pass].W7RNeg[2] = -(2 * FMath::Cos(Phase3 + Phase3) - 1);
+					FinalWeights[Pass].W7RNeg[3] = -(2 * FMath::Cos(Phase3 + Phase3) - 1);
+				}
+			}
+
+			void GenerateRadix4Weights()
+			{
+				const int32 Radix4PassSize = 16;
+				check(FFTSize >= Radix4PassSize);
+
+				// Each radix-4 butterfly
+				int32 MaxNumButterfliesInStage = FFTSize / Radix4PassSize;
+
+				Radix4Weights.Reset();
+				Radix4Weights.AddUninitialized(MaxNumButterfliesInStage);
+
+				for (int32 ButterflyIndex = 0; ButterflyIndex < MaxNumButterfliesInStage; ++ButterflyIndex)
+				{
+					const double Phase = 2. * PI * RotateBitsAroundPoint(4 * ButterflyIndex);
+
+					const float W1R = static_cast<float>(FMath::Cos(Phase));
+					const float W1I = static_cast<float>(FMath::Tan(Phase));
+					const float W2R = static_cast<float>(FMath::Cos(Phase + Phase));
+					const float W2I = static_cast<float>(FMath::Tan(Phase + Phase));
+					const float W3R = static_cast<float>(2. * W2R - 1.);
+					const float W3I = static_cast<float>(FMath::Tan(3. * Phase));
+
+
+					Radix4Weights[ButterflyIndex].W1R[0] = W1R;
+					Radix4Weights[ButterflyIndex].W1I[0] = -W1I;
+					Radix4Weights[ButterflyIndex].W2R[0] = W2R;
+					Radix4Weights[ButterflyIndex].W2I[0] = -W2I;
+					Radix4Weights[ButterflyIndex].W3R[0] = W3R;
+					Radix4Weights[ButterflyIndex].W3I[0] = -W3I;
+
+					Radix4Weights[ButterflyIndex].W1RNeg[0] = -W1R;
+					Radix4Weights[ButterflyIndex].W1RD2[0] = -W1R;
+					Radix4Weights[ButterflyIndex].W1RD3[0] = W1R;
+					Radix4Weights[ButterflyIndex].W2RNeg[0] = -W2R;
+					Radix4Weights[ButterflyIndex].W3RNeg[0] = -W3R;
+
+
+					Radix4Weights[ButterflyIndex].W1R[1] = W1R;
+					Radix4Weights[ButterflyIndex].W1I[1] = W1I;
+					Radix4Weights[ButterflyIndex].W2R[1] = W2R;
+					Radix4Weights[ButterflyIndex].W2I[1] = W2I;
+					Radix4Weights[ButterflyIndex].W3R[1] = W3R;
+					Radix4Weights[ButterflyIndex].W3I[1] = W3I;
+
+					Radix4Weights[ButterflyIndex].W1RNeg[1] = -W1R;
+					Radix4Weights[ButterflyIndex].W1RD2[1] = W1R;
+					Radix4Weights[ButterflyIndex].W1RD3[1] = -W1R;
+					Radix4Weights[ButterflyIndex].W2RNeg[1] = -W2R;
+					Radix4Weights[ButterflyIndex].W3RNeg[1] = -W3R;
+
+
+					Radix4Weights[ButterflyIndex].W1R[2] = W1R;
+					Radix4Weights[ButterflyIndex].W1I[2] = -W1I;
+					Radix4Weights[ButterflyIndex].W2R[2] = W2R;
+					Radix4Weights[ButterflyIndex].W2I[2] = -W2I;
+					Radix4Weights[ButterflyIndex].W3R[2] = W3R;
+					Radix4Weights[ButterflyIndex].W3I[2] = -W3I;
+
+					Radix4Weights[ButterflyIndex].W1RNeg[2] = -W1R;
+					Radix4Weights[ButterflyIndex].W1RD2[2] = -W1R;
+					Radix4Weights[ButterflyIndex].W1RD3[2] = W1R;
+					Radix4Weights[ButterflyIndex].W2RNeg[2] = -W2R;
+					Radix4Weights[ButterflyIndex].W3RNeg[2] = -W3R;
+
+
+					Radix4Weights[ButterflyIndex].W1R[3] = W1R;
+					Radix4Weights[ButterflyIndex].W1I[3] = W1I;
+					Radix4Weights[ButterflyIndex].W2R[3] = W2R;
+					Radix4Weights[ButterflyIndex].W2I[3] = W2I;
+					Radix4Weights[ButterflyIndex].W3R[3] = W3R;
+					Radix4Weights[ButterflyIndex].W3I[3] = W3I;
+
+					Radix4Weights[ButterflyIndex].W1RNeg[3] = -W1R;
+					Radix4Weights[ButterflyIndex].W1RD2[3] = W1R;
+					Radix4Weights[ButterflyIndex].W1RD3[3] = -W1R;
+					Radix4Weights[ButterflyIndex].W2RNeg[3] = -W2R;
+					Radix4Weights[ButterflyIndex].W3RNeg[3] = -W3R;
+				}
+			}
+
+			// A funny but useful function which reverses bits and
+			// places them behind a point.  For example
+			// 0101 is transformed to 0.1010
+			double RotateBitsAroundPoint(uint32 InValue)
+			{
+				uint32 ReversedValue = ReverseBits(InValue);
+				double OutValue = ReversedValue / 4294967296.;
+				OutValue = 1. / 4294967296. * ReversedValue;
+				return OutValue;
+			}
+
+			uint32 ReverseBits(uint32 InValue)
+			{
+				static const uint8 ByteReversal[256] = {
+					 0, 128, 64, 192, 32, 160, 96, 224, 16, 144, 80, 208, 48, 176, 112, 240,
+					 8, 136, 72, 200, 40, 168, 104, 232, 24, 152, 88, 216, 56, 184, 120, 248,
+					 4, 132, 68, 196, 36, 164, 100, 228, 20, 148, 84, 212, 52, 180, 116, 244,
+					 12, 140, 76, 204, 44, 172, 108, 236, 28, 156, 92, 220, 60, 188, 124, 252,
+					 2, 130, 66, 194, 34, 162, 98, 226, 18, 146, 82, 210, 50, 178, 114, 242,
+					 10, 138, 74, 202, 42, 170, 106, 234, 26, 154, 90, 218, 58, 186, 122, 250,
+					6, 134, 70, 198, 38, 166, 102, 230, 22, 150, 86, 214, 54, 182, 118, 246,
+					 14, 142, 78, 206, 46, 174, 110, 238, 30, 158, 94, 222, 62, 190, 126, 254,
+					 1, 129, 65, 193, 33, 161, 97, 225, 17, 145, 81, 209, 49, 177, 113, 241,
+					 9, 137, 73, 201, 41, 169, 105, 233, 25, 153, 89, 217, 57, 185, 121, 249,
+					 5, 133, 69, 197, 37, 165, 101, 229, 21, 149, 85, 213, 53, 181, 117, 245,
+					 13, 141, 77, 205, 45, 173, 109, 237, 29, 157, 93, 221, 61, 189, 125, 253,
+					 3, 131, 67, 195, 35, 163, 99, 227, 19, 147, 83, 211, 51, 179, 115, 243,
+					 11, 139, 75, 203, 43, 171, 107, 235, 27, 155, 91, 219, 59, 187, 123, 251,
+					 7, 135, 71, 199, 39, 167, 103, 231, 23, 151, 87, 215, 55, 183, 119, 247,
+					 15, 143, 79, 207, 47, 175, 111, 239, 31, 159, 95, 223, 63, 191, 127, 255
+				};
+
+				uint8 Byte0 = ByteReversal[InValue >> 0*8 & 0xff];
+				uint8 Byte1 = ByteReversal[InValue >> 1*8 & 0xff];
+				uint8 Byte2 = ByteReversal[InValue >> 2*8 & 0xff];
+				uint8 Byte3 = ByteReversal[InValue >> 3*8 & 0xff];
+
+				uint32 OutValue = Byte0 << 3*8 | Byte1 << 2*8 | Byte2 << 1*8 | Byte3 << 0*8;
+				return OutValue;
+			}
+			
+			int32 Log2FFTSize;
+			int32 FFTSize;
+			int32 NumFloats;
+
+			TArray<FRadix4Weight> Radix4Weights;
+			TArray<FFinalIndices> FinalIndices;
+			TArray<FFinalWeights> FinalWeights;
+
+			AlignedFloatBuffer InverseWorkBuffer;
+	};
+
+	// Maximum log 2 size of fft
+	const int32 FVectorRealToComplexFFT::MinLog2FFTSize = FVectorComplexFFT::MinLog2FFTSize + 1;
+
+	// Maximum log 2 size of fft
+	const int32 FVectorRealToComplexFFT::MaxLog2FFTSize = FVectorComplexFFT::MaxLog2FFTSize + 1;
+
+	void FVectorRealToComplexFFT::InitRealSequenceConversionBuffers()
+	{
+		// Conversion buffers for performing a real valued FFT using a complex fft
+		// The values in the buffer are setup to support SIMD operations resulting
+		// in some duplicate data.
+		ForwardConvBuffers.AlphaReal.AddUninitialized(FFTSize);
+		ForwardConvBuffers.AlphaImag.AddUninitialized(FFTSize);
+		ForwardConvBuffers.BetaReal.AddUninitialized(FFTSize);
+		ForwardConvBuffers.BetaImag.AddUninitialized(FFTSize);
+
+		InverseConvBuffers.AlphaReal.AddUninitialized(FFTSize);
+		InverseConvBuffers.AlphaImag.AddUninitialized(FFTSize);
+		InverseConvBuffers.BetaReal.AddUninitialized(FFTSize);
+		InverseConvBuffers.BetaImag.AddUninitialized(FFTSize);
+
+		float* AlphaRealForwardBufferData = ForwardConvBuffers.AlphaReal.GetData();
+		float* AlphaImagForwardBufferData = ForwardConvBuffers.AlphaImag.GetData();
+		float* BetaRealForwardBufferData = ForwardConvBuffers.BetaReal.GetData();
+		float* BetaImagForwardBufferData = ForwardConvBuffers.BetaImag.GetData();
+
+		float* AlphaRealInverseBufferData = InverseConvBuffers.AlphaReal.GetData();
+		float* AlphaImagInverseBufferData = InverseConvBuffers.AlphaImag.GetData();
+		float* BetaRealInverseBufferData = InverseConvBuffers.BetaReal.GetData();
+		float* BetaImagInverseBufferData = InverseConvBuffers.BetaImag.GetData();
+
+		float PhaseIncrement = PI / static_cast<float>(FFTSize);
+
+		for (int32 i = 0; i < FFTSize; i += 2)
+		{
+			const float Phase = PhaseIncrement * i;
+			const float BetaReal = 0.5 * (1. - FMath::Sin(Phase));
+			const float BetaImag = -0.5 * FMath::Cos(Phase);
+			const float AlphaReal = 0.5 * (1. + FMath::Sin(Phase));
+			const float AlphaImag = 0.5 * FMath::Cos(Phase);
+
+			AlphaRealForwardBufferData[i] 		= AlphaReal;
+			AlphaRealForwardBufferData[i + 1] 	= -AlphaReal;// Sign flipped to simplify SIMD math
+			AlphaImagForwardBufferData[i] 		= AlphaImag; 
+			AlphaImagForwardBufferData[i + 1] 	= AlphaImag;
+
+			BetaRealForwardBufferData[i] 		= BetaReal;
+			BetaRealForwardBufferData[i + 1] 	= BetaReal;
+			BetaImagForwardBufferData[i] 		= -BetaImag; // Sign flipped to simplify SIMD math 
+			BetaImagForwardBufferData[i + 1] 	= BetaImag;
+			
+			AlphaRealInverseBufferData[i] 		= AlphaReal;
+			AlphaRealInverseBufferData[i + 1] 	= -AlphaReal; // Sign flipped to simplify SIMD math 
+			AlphaImagInverseBufferData[i] 		= AlphaImag;
+			AlphaImagInverseBufferData[i + 1] 	= AlphaImag; 
+
+			BetaRealInverseBufferData[i] 		= BetaReal;
+			BetaRealInverseBufferData[i + 1] 	= BetaReal;
+			BetaImagInverseBufferData[i] 		= -BetaImag; // Sign flipped to simplify SIMD math 
+			BetaImagInverseBufferData[i + 1] 	= BetaImag; 
+		}
+	}
+
+	// Performs conversion of buffers required to do real fft using complex fft.
+	void FVectorRealToComplexFFT::ConvertSequence(const FConversionBuffers& InBuffers, const float* RESTRICT InValues, int32 InStartIndex, float* RESTRICT OutValues)
+	{
+		const float* AlphaRealData = InBuffers.AlphaReal.GetData();
+		const float* AlphaImagData = InBuffers.AlphaImag.GetData();
+		const float* BetaRealData = InBuffers.BetaReal.GetData();
+		const float* BetaImagData = InBuffers.BetaImag.GetData();
+
+		if (FFTSize > InStartIndex)
+		{
+			VectorRegister VInRev1 = VectorLoadAligned(&InValues[FFTSize - InStartIndex]);
+
+			for (int32 i = InStartIndex; i < FFTSize; i += 4)
+			{
+				VectorRegister VIn = VectorLoadAligned(&InValues[i]);
+				VectorRegister VInRISwap = VectorSwizzle(VIn, 1, 0, 3, 2);
+
+				VectorRegister VInRev2 = VectorLoadAligned(&InValues[FFTSize - i - 4]);
+				VectorRegister VInRev = VectorShuffle(VInRev1, VInRev2, 0, 1, 2, 3);
+				VInRev1 = VInRev2;
+
+				VectorRegister VInRevRISwap = VectorSwizzle(VInRev, 1, 0, 3, 2);
+
+				VectorRegister VAlphaReal = VectorLoadAligned(&AlphaRealData[i]);
+				VectorRegister VAlphaImag = VectorLoadAligned(&AlphaImagData[i]);
+				VectorRegister VBetaReal = VectorLoadAligned(&BetaRealData[i]);
+				VectorRegister VBetaImag = VectorLoadAligned(&BetaImagData[i]);
+
+				// Out1 = [ R * Ar,  I * Ar]
+				// Out2 = [ I * Ai,  R * Ai]
+				// Out3 = [NR * Br, NI * Br]
+				// Out4 = [NI * Bi, NR * Bi]
+				VectorRegister Out1 = VectorMultiply(VIn, VAlphaReal);
+				VectorRegister Out2 = VectorMultiply(VInRISwap, VAlphaImag);
+				VectorRegister Out3 = VectorMultiply(VInRev, VBetaReal);
+				VectorRegister Out4 = VectorMultiply(VInRevRISwap, VBetaImag);
+
+				// Out12 = [(R * Ar) + (I * Ai), (I * Ar) + (R * Ai)]
+				VectorRegister Out12 = VectorAdd(Out1, Out2);
+				// Out34 = [(NR * Br) + (NI * Bi), (NR * Bi) + (NI * Br)]
+				VectorRegister Out34 = VectorAdd(Out3, Out4);
+
+				// Out = [
+				// 	(R * Ar) + (I * Ai) + (NR * Br) + (NI * Bi),
+				// 	(I * Ar) + (R * Ai) + (NR * Bi) + (NI * Br)
+				// ]
+				VectorRegister Out = VectorAdd(Out12, Out34);
+				VectorStoreAligned(Out, &OutValues[i]);
+			}
+		}
+	}
+
+	FVectorRealToComplexFFT::FVectorRealToComplexFFT(int32 InLog2FFTSize)
+	:	FFTSize(1 << InLog2FFTSize)
+	,	Log2FFTSize(InLog2FFTSize)
+	,	ComplexFFT(new FVectorComplexFFT(InLog2FFTSize - 1)) // Utilize a N/2 length complex fft to perform N length fft.
+	{
+		WorkBuffer.AddUninitialized(FFTSize);
+
+		InitRealSequenceConversionBuffers();
+	}
+
+	FVectorRealToComplexFFT::~FVectorRealToComplexFFT()
+	{
+	}
+
+	int32 FVectorRealToComplexFFT::Size() const
+	{
+		return FFTSize;
+	}
+
+	/** Scaling applied when performing forward FFT. */
+	EFFTScaling FVectorRealToComplexFFT::ForwardScaling() const
+	{
+		return EFFTScaling::None;
+	}
+
+	/** Scaling applied when performing inverse FFT. */
+	EFFTScaling FVectorRealToComplexFFT::InverseScaling() const
+	{
+		return EFFTScaling::None;
+	}
+
+	void FVectorRealToComplexFFT::ForwardRealToComplex(const float* RESTRICT InReal, float* RESTRICT OutComplex)
+	{
+		checkf(IsAligned<const float*>(InReal, AUDIO_SIMD_FLOAT_ALIGNMENT), TEXT("Memory must be aligned.")); 
+		checkf(IsAligned<float*>(OutComplex, AUDIO_SIMD_FLOAT_ALIGNMENT), TEXT("Memory must be aligned.")); 
+
+		// Performs a N sized real-to-complex FFT using an N/2 complex-to-complex FFT.
+		float* WorkData = WorkBuffer.GetData();
+
+		ComplexFFT->ForwardComplexToComplex(InReal, WorkData);
+
+		const float* AlphaRealForwardData = ForwardConvBuffers.AlphaReal.GetData();
+		const float* AlphaImagForwardData = ForwardConvBuffers.AlphaImag.GetData();
+		const float* BetaRealForwardData = ForwardConvBuffers.BetaReal.GetData();
+		const float* BetaImagForwardData = ForwardConvBuffers.BetaImag.GetData();
+
+		// Handle special case of this math to account for cyclical index math.
+		OutComplex[0] = (WorkData[0] * AlphaRealForwardData[0]) 
+			+ (WorkData[1] * AlphaImagForwardData[0]) 
+			+ (WorkData[0] * BetaRealForwardData[0]) 
+			+ (WorkData[1] * BetaImagForwardData[0]);
+
+		OutComplex[1] = (WorkData[1] * AlphaRealForwardData[1]) 
+			+ (WorkData[0] * AlphaImagForwardData[1]) 
+			+ (WorkData[0] * BetaImagForwardData[1]) 
+			+ (WorkData[1] * BetaRealForwardData[1]);
+
+		OutComplex[2] = (WorkData[2] * AlphaRealForwardData[2]) 
+			+ (WorkData[3] * AlphaImagForwardData[2]) 
+			+ (WorkData[FFTSize - 2] * BetaRealForwardData[2]) 
+			+ (WorkData[FFTSize - 1] * BetaImagForwardData[2]);
+		
+
+		OutComplex[3] = (WorkData[3] * AlphaRealForwardData[3]) 
+			+ (WorkData[2] * AlphaImagForwardData[3]) 
+			+ (WorkData[FFTSize - 2] * BetaImagForwardData[3]) 
+			+ (WorkData[FFTSize - 1] * BetaRealForwardData[3]);
+
+		// Convert all other values using optimized SIMD
+		ConvertSequence(ForwardConvBuffers, WorkData, 4, OutComplex);
+
+		// Handle special case of nyquist frequency
+		OutComplex[FFTSize] = WorkData[0] - WorkData[1];
+		OutComplex[FFTSize + 1] = 0.f;
+	}
+
+	void FVectorRealToComplexFFT::InverseComplexToReal(const float* RESTRICT InComplex, float* RESTRICT OutReal)
+	{
+		checkf(IsAligned<const float*>(InComplex, AUDIO_SIMD_FLOAT_ALIGNMENT), TEXT("Memory must be aligned.")); 
+		checkf(IsAligned<float*>(OutReal, AUDIO_SIMD_FLOAT_ALIGNMENT), TEXT("Memory must be aligned.")); 
+
+		// Performs a N sized complex-to-real FFT using an N/2 complex-to-complex FFT.
+
+		float* WorkData = WorkBuffer.GetData();
+
+		const float* AlphaRealInverseData = InverseConvBuffers.AlphaReal.GetData();
+		const float* AlphaImagInverseData = InverseConvBuffers.AlphaImag.GetData();
+		const float* BetaRealInverseData = InverseConvBuffers.BetaReal.GetData();
+		const float* BetaImagInverseData = InverseConvBuffers.BetaImag.GetData();
+
+		// Handle special case of this math to account for cyclical index math.
+		WorkData[0] = (InComplex[0] * AlphaRealInverseData[0]) 
+			+ (InComplex[1] * AlphaImagInverseData[0]) 
+			+ (InComplex[FFTSize] * BetaRealInverseData[0]) 
+			+ (InComplex[FFTSize + 1] * BetaImagInverseData[0]);
+
+		WorkData[1] = (InComplex[1] * AlphaRealInverseData[1]) 
+			+ (InComplex[0] * AlphaImagInverseData[1]) 
+			+ (InComplex[FFTSize] * BetaImagInverseData[1]) 
+			+ (InComplex[FFTSize + 1] * BetaRealInverseData[0]);
+
+		WorkData[2] = (InComplex[2] * AlphaRealInverseData[2])
+			+ (InComplex[3] * AlphaImagInverseData[2])
+			+ (InComplex[FFTSize - 2] * BetaRealInverseData[2])
+			+ (InComplex[FFTSize - 1] * BetaImagInverseData[2]);
+
+		WorkData[3] = (InComplex[3] * AlphaRealInverseData[3]) 
+			+ (InComplex[2] * AlphaImagInverseData[3]) 
+			+ (InComplex[FFTSize - 2] * BetaImagInverseData[3]) 
+			+ (InComplex[FFTSize - 1] * BetaRealInverseData[3]);
+
+		// Convert all other values using optimized SIMD
+		ConvertSequence(ForwardConvBuffers, InComplex, 4, WorkData);
+
+		// Perform Inverse FFT 
+		ComplexFFT->InverseComplexToComplex(WorkData, OutReal);
+	}
+
+	void FVectorRealToComplexFFT::BatchForwardRealToComplex(int32 InCount, const float* const RESTRICT InReal[], float* RESTRICT OutComplex[])
+	{
+		for (int32 i = 0; i < InCount; i++)
+		{
+			ForwardRealToComplex(InReal[i], OutComplex[i]);
+		}
+	}
+
+	void FVectorRealToComplexFFT::BatchInverseComplexToReal(int32 InCount, const float* const RESTRICT InComplex[], float* RESTRICT OutReal[])
+	{
+		for (int32 i = 0; i < InCount; i++)
+		{
+			InverseComplexToReal(InComplex[i], OutReal[i]);
+		}
+	}
+
+	/*************************************************************************************************/
+	/**************************************** FVectorFFTFactory **************************************/
+	/*************************************************************************************************/
+	FVectorFFTFactory::~FVectorFFTFactory()
+	{
+	}
+
+	/** Name of this particular factory. */
+	FName FVectorFFTFactory::GetFactoryName() const
+	{
+		static const FName FactoryName = FName(TEXT("FVectorFFTFactory"));
+		return FactoryName;
+	}
+
+	/** If true, this implementation uses hardware acceleration. */
+	bool FVectorFFTFactory::IsHardwareAccelerated() const
+	{
+		return false;
+	}
+
+	/** If true, this implementation requires input and output arrays to be 128 bit aligned. */
+	bool FVectorFFTFactory::Expects128BitAlignedArrays() const
+	{
+		return true;
+	}
+
+	/** Returns true if the input settings are supported by this factory. */
+	bool FVectorFFTFactory::AreFFTSettingsSupported(const FFFTSettings& InSettings) const
+	{
+		// Supports FFT sizes of 5 to 16, though an FFT 
+		bool bIsMinSizeSupported = InSettings.Log2Size >= FVectorRealToComplexFFT::MinLog2FFTSize;
+		bool bIsMaxSizeSupported = InSettings.Log2Size <= FVectorRealToComplexFFT::MaxLog2FFTSize;
+		bool bIsAlignmentSupported = InSettings.bArrays128BitAligned;
+
+		return bIsMinSizeSupported && bIsAlignmentSupported && bIsMaxSizeSupported;
+	}
+
+	/** Creates a new FFT algorithm. */
+	TUniquePtr<IFFTAlgorithm> FVectorFFTFactory::NewFFTAlgorithm(const FFFTSettings& InSettings)
+	{
+		if (AreFFTSettingsSupported(InSettings))
+		{
+			return MakeUnique<FVectorRealToComplexFFT>(InSettings.Log2Size);
+		}
+		return TUniquePtr<IFFTAlgorithm>();
+	}
+}
diff --git a/Engine/Source/Runtime/SignalProcessing/Private/VectorFFT.h b/Engine/Source/Runtime/SignalProcessing/Private/VectorFFT.h
new file mode 100644
--- /dev/null
+++ b/Engine/Source/Runtime/SignalProcessing/Private/VectorFFT.h
@@ -0,0 +1,129 @@
+// Copyright Epic Games, Inc. All Rights Reserved.
+
+#include "DSP/FFTAlgorithm.h"
+#include "DSP/BufferVectorOperations.h"
+
+#pragma once
+
+namespace Audio
+{
+	// Forward declaration
+	class FVectorComplexFFT;
+
+	class FVectorRealToComplexFFT : public IFFTAlgorithm
+	{
+		public:
+			/** Minimum log2 size of fft. */
+			static const int32 MinLog2FFTSize;
+
+			/** Maximum log 2 size of fft */
+			static const int32 MaxLog2FFTSize;
+
+			/** Construct a FVectorRealToComplex FFT
+			 *
+			 * @param InLog2Size - Determines the size of the FFT.  FFTSize = 2^InLog2Size.
+			 */
+			FVectorRealToComplexFFT(int32 InLog2Size);
+
+			virtual ~FVectorRealToComplexFFT();
+
+			/** Number of elements in FFT */
+			virtual int32 Size() const override;
+
+			/** Scaling applied when performing forward FFT. */
+			virtual EFFTScaling ForwardScaling() const override;
+
+			/** Scaling applied when performing inverse FFT. */
+			virtual EFFTScaling InverseScaling() const override;
+
+			
+			/** ForwardRealToComplex
+			 *
+			 * Perform the forward FFT on real data and produce complex data.
+			 *
+			 * @param InReal - Array of floats to input into fourier transform. Must have FFTSize() elements.
+			 * @param OutComplex - Array of floats to store output of fourier transform. Must have (FFTSize() + 2) float elements which represent ((FFTSize() / 2) + 1) complex numbers in interleaved format.
+			 */
+			virtual void ForwardRealToComplex(const float* RESTRICT InReal, float* RESTRICT OutComplex) override;
+
+			/** InverseComplexToReal
+			 *
+			 * Perform the inverse FFT on complex data and produce real data.
+			 *
+			 * @param InComplex - Array of floats to input into inverse fourier transform. Must have (FFTSize() + 2) float elements which represent ((FFTSize() / 2) + 1) complex numbers in interleaved format.
+			 * @param OutReal - Array of floats to store output of inverse fourier transform. Must have FFTSize() elements.
+			 */
+			virtual void InverseComplexToReal(const float* RESTRICT InComplex, float* RESTRICT OutReal) override;
+
+			/** BatchForwardRealToComplex
+			 *
+			 * Perform the forward FFT on real data and produce complex data.
+			 *
+			 * @param InCount - The number of transforms to compute.
+			 * @param InReal - Array of array of floats to input into fourier transform. Must have InCount elements, each containing FFTSize() elements.
+			 * @param OutComplex - Array of array of floats to store output of fourier transform. Must have InCount elements, each with (FFTSize() + 2) float elements which represent ((FFTSize() / 2) + 1) complex numbers in interleaved format.
+			 */
+			virtual void BatchForwardRealToComplex(int32 InCount, const float* const RESTRICT InReal[], float* RESTRICT OutComplex[]) override;
+
+			/** BatchInverseComplexToReal
+			 *
+			 * Perform the inverse FFT on complex data and produce real data.
+			 *
+			 * @param InCount - The number of transforms to compute.
+			 * @param InComplex - Array of array of floats to input into inverse fourier transform. Must have InCount elements, each with (FFTSize() + 2) float elements which represent ((FFTSize() / 2) + 1) complex numbers in interleaved format.
+			 * @param OutReal - Array of array of floats to store output of inverse fourier transform. Must have InCount elements, each with FFTSize() float elements.
+			 */
+			virtual void BatchInverseComplexToReal(int32 InCount, const float* const RESTRICT InComplex[], float* RESTRICT OutReal[]) override;
+
+		private:
+
+			// Needed to convert buffers for implementation of real-to-complex fft using complex-to-comlex fft.
+			struct FConversionBuffers
+			{
+				AlignedFloatBuffer AlphaReal;
+				AlignedFloatBuffer AlphaImag;
+				AlignedFloatBuffer BetaReal;
+				AlignedFloatBuffer BetaImag;
+			};
+
+
+			// Initializes conversion buffers
+			void InitRealSequenceConversionBuffers();
+
+			// Converts a buffer
+			void ConvertSequence(const FConversionBuffers& InBuffers, const float* RESTRICT InValues, int32 InStartIndex, float* RESTRICT OutValues);
+
+			int32 FFTSize;
+			int32 Log2FFTSize;
+
+			AlignedFloatBuffer WorkBuffer;
+
+			FConversionBuffers ForwardConvBuffers;
+			FConversionBuffers InverseConvBuffers;
+
+			// Complex FFT implementation.
+			TUniquePtr<FVectorComplexFFT> ComplexFFT;
+	};
+
+	/** Algorithm factory for FVectorFFT */
+	class FVectorFFTFactory : public IFFTAlgorithmFactory
+	{
+		public:
+			virtual ~FVectorFFTFactory();
+
+			/** Name of this particular factory. */
+			virtual FName GetFactoryName() const override;
+
+			/** If true, this implementation uses hardware acceleration. */
+			virtual bool IsHardwareAccelerated() const override;
+
+			/** If true, this implementation requires input and output arrays to be 128 bit aligned. */
+			virtual bool Expects128BitAlignedArrays() const override;
+
+			/** Returns true if the input settings are supported by this factory. */
+			virtual bool AreFFTSettingsSupported(const FFFTSettings& InSettings) const override;
+
+			/** Creates a new FFT algorithm. */
+			virtual TUniquePtr<IFFTAlgorithm> NewFFTAlgorithm(const FFFTSettings& InSettings) override;
+	};
+}
diff --git a/Engine/Source/Runtime/SignalProcessing/Public/DSP/AudioFFT.h b/Engine/Source/Runtime/SignalProcessing/Public/DSP/AudioFFT.h
--- a/Engine/Source/Runtime/SignalProcessing/Public/DSP/AudioFFT.h
+++ b/Engine/Source/Runtime/SignalProcessing/Public/DSP/AudioFFT.h
@@ -5,6 +5,7 @@
 #include "CoreMinimal.h"
 #include "DSP/BufferVectorOperations.h"
 #include "DSP/Dsp.h"
+#include "DSP/FFTAlgorithm.h"
 
 namespace Audio
 {
@@ -143,6 +144,29 @@ namespace Audio
 	SIGNALPROCESSING_API void PerformFFT(const FFTTimeDomainData& InputParams, FFTFreqDomainData& OutputParams);
 	SIGNALPROCESSING_API void PerformIFFT(FFTFreqDomainData& InputParams, FFTTimeDomainData& OutputParams);
 
+
+	// FFT Algorithm factory for this FFT implementation
+	class SIGNALPROCESSING_API FAudioFFTAlgorithmFactory : public IFFTAlgorithmFactory
+	{
+		public:
+			virtual ~FAudioFFTAlgorithmFactory();
+
+			// Name of this fft algorithm factory. 
+			virtual FName GetFactoryName() const override;
+
+			// If true, this implementation uses hardware acceleration.
+			virtual bool IsHardwareAccelerated() const override;
+
+			// If true, this implementation requires input and output arrays to be 128 bit aligned.
+			virtual bool Expects128BitAlignedArrays() const override;
+
+			// Returns true if the input settings are supported by this factory.
+			virtual bool AreFFTSettingsSupported(const FFFTSettings& InSettings) const override;
+
+			// Create a new FFT algorithm.
+			virtual TUniquePtr<IFFTAlgorithm> NewFFTAlgorithm(const FFFTSettings& InSettings) override;
+	};
+
 	struct FrequencyBuffer
 	{
 		AlignedFloatBuffer Real;
diff --git a/Engine/Source/Runtime/SignalProcessing/Public/DSP/ConvolutionAlgorithm.h b/Engine/Source/Runtime/SignalProcessing/Public/DSP/ConvolutionAlgorithm.h
new file mode 100644
--- /dev/null
+++ b/Engine/Source/Runtime/SignalProcessing/Public/DSP/ConvolutionAlgorithm.h
@@ -0,0 +1,140 @@
+// Copyright Epic Games, Inc. All Rights Reserved.
+
+#pragma once
+
+#include "CoreMinimal.h"
+#include "Features/IModularFeature.h"
+
+namespace Audio
+{
+	/** FConvolutionSettings
+	 *
+	 * Settings for creating an IConvolutionAlgorithm.
+	 */
+	struct FConvolutionSettings
+	{
+		/** If true, hardware accelerated algorithms are valid. */
+		bool bEnableHardwareAcceleration = true;
+
+		/** Defines size of audio processing block. */
+		int32 BlockNumSamples = 256;
+
+		/** Number of input audio channels to support. */
+		int32 NumInputChannels = 0;
+
+		/** Number of output audio channels to support. */
+		int32 NumOutputChannels = 0;
+
+		/** Number of impulse responses to support. */
+		int32 NumImpulseResponses = 0;
+
+		/** Maximum number of samples in an impulse responses. */
+		int32 MaxNumImpulseResponseSamples = 0;
+	};
+
+
+	/** IConvolutionAlgorithm 
+	 *
+	 * Interface for Convolution algorithm.
+	 */
+	class SIGNALPROCESSING_API IConvolutionAlgorithm 
+	{
+		public:
+
+			/** virtual destructor for inheritance. */
+			virtual ~IConvolutionAlgorithm() {}
+
+			/** Returns the number of samples in an audio block. */
+			virtual int32 GetNumSamplesInBlock() const = 0;
+
+			/** Returns number of audio inputs. */
+			virtual int32 GetNumAudioInputs() const = 0;
+
+			/** Returns number of audio outputs. */
+			virtual int32 GetNumAudioOutputs() const = 0;
+
+			/** Process one block of audio.
+			 *
+			 * InSamples is processed by the impulse responses. The output is placed in OutSamples.
+			 *
+			 * @params InSamples - A 2D array of input deinterleaved audio samples. InSamples[GetNumAudioInputs()][GetNumSamplesInBlock()]
+			 * @params OutSamples - A 2D array of output deinterleaved audio samples. OutSamples[GetNumAudioOutputs()][GetNumSamplesInBlock()]
+			 *
+			 */
+			virtual void ProcessAudioBlock(const float* const InSamples[], float* const OutSamples[]) = 0;
+
+			/** Reset internal history buffers. */
+			virtual void ResetAudioHistory() = 0;
+
+			/** Maximum supported length of impulse response. */
+			virtual int32 GetMaxNumImpulseResponseSamples() const = 0;
+
+
+			/** Return the number of impulse responses. */
+			virtual int32 GetNumImpulseResponses() const = 0;
+
+			/** Return the number of samples in an impulse response. */
+			virtual int32 GetNumImpulseResponseSamples(int32 InImpulseResponseIndex) const = 0;
+
+			/** Set impulse response values. */
+			virtual void SetImpulseResponse(int32 InImpulseResponseIndex, const float* InSamples, int32 NumSamples) = 0;
+
+			/** Sets the gain between an audio input, impulse response and audio output.
+			 *
+			 * ([audio input] * [impulse response]) x gain = [audio output]
+			 */
+			virtual void SetMatrixGain(int32 InAudioInputIndex, int32 InImpulseResponseIndex, int32 InAudioOutputIndex, float InGain) = 0;
+
+			/** Gets the gain between an audio input, impulse response and audio output.
+			 *
+			 * ([audio input] * [impulse response]) x gain = [audio output]
+			 */
+			virtual float GetMatrixGain(int32 InAudioInputIndex, int32 InImpulseResponseIndex, int32 InAudioOutputIndex) const = 0;
+	};
+
+	/** IConvolutionAlgorithmFactory
+	 *
+	 * Factory interface for creating convolution algorithms.
+	 */
+	class SIGNALPROCESSING_API IConvolutionAlgorithmFactory : public IModularFeature
+	{
+		public:
+			virtual ~IConvolutionAlgorithmFactory();
+
+			/** Name of modular feature for Convolution factory.  */
+			static const FName GetModularFeatureName();
+
+			/** Name of this particular factory. */
+			virtual const FName GetFactoryName() const = 0;
+
+			/** If true, this implementation uses hardware acceleration. */
+			virtual bool IsHardwareAccelerated() const = 0;
+
+			/** Returns true if the input settings are supported by this factory. */
+			virtual bool AreConvolutionSettingsSupported(const FConvolutionSettings& InSettings) const = 0;
+
+			/** Creates a new Convolution algorithm. */
+			virtual TUniquePtr<IConvolutionAlgorithm> NewConvolutionAlgorithm(const FConvolutionSettings& InSettings) = 0;
+	};
+
+	/** FConvolutionFactory
+	 *
+	 * FConvolutionFactory creates convolution algorithms. It will choose hardware accelerated versions when they are available. 
+	 */
+	class SIGNALPROCESSING_API FConvolutionFactory
+	{
+		public:
+			/** This denotes that no specific IConvolutionAlgorithmFactory is desired. */
+			static const FName AnyAlgorithmFactory;
+
+			/** NewConvolutionAlgorithm
+			 *
+			 * Creates and returns a new ConvolutionAlgorithm. 
+			 *
+			 * @param InSettings - The settings used to create the Convolution algorithm.
+			 * @param InAlgorithmFactoryName - If not equal to FConvolutionFactory::AnyAlgorithmFactory, will only uses Convolution algorithm facotry if IConvolutionAlgorithmFactory::GetFactoryName() equals InAlgorithmFactoryName.
+			 * @return A TUniquePtr<IConvolutionAlgorithm> to the created Convolution. This pointer can be invalid if an error occured or the fft algorithm could not be created.
+			 */
+			static TUniquePtr<IConvolutionAlgorithm> NewConvolutionAlgorithm(const FConvolutionSettings& InSettings, const FName& InAlgorithmFactoryName = AnyAlgorithmFactory);
+	};
+}
diff --git a/Engine/Source/Runtime/SignalProcessing/Public/DSP/FFTAlgorithm.h b/Engine/Source/Runtime/SignalProcessing/Public/DSP/FFTAlgorithm.h
new file mode 100644
--- /dev/null
+++ b/Engine/Source/Runtime/SignalProcessing/Public/DSP/FFTAlgorithm.h
@@ -0,0 +1,175 @@
+// Copyright Epic Games, Inc. All Rights Reserved.
+
+#pragma once
+
+#include "CoreMinimal.h"
+#include "Features/IModularFeature.h"
+
+namespace Audio
+{
+	/** EFFTScaling
+	 *
+	 * EFFTScaling describes how transformed data is scaled. IFFTAlgorithm implementations
+	 * must return this value to describe how to achieve equal energy between input data and 
+	 * transformed data. It is assumed that applying the logical inverse of the scaling will
+	 * produce equal energy.
+	 */
+	enum class EFFTScaling : uint8
+	{
+		/** No scaling needed to maintain equal energy. */
+		None,
+
+		/** Output is effectively multiplied by FFTSize. To maintain equal energy, output must be divided by the FFTSize. */
+		MultipliedByFFTSize,
+
+		/** Output is effectively multiplied by the square root of the FFTSize. To maintain equal energy, output must be divided by the square root of the FFTSize. */
+		MultipliedBySqrtFFTSize,
+
+		/** Output is effectively divided by the square root of the FFTSize. To maintain equal energy, output must be multiplied by the square root of the FFTSize. */
+		DividedByFFTSize,
+
+		/** Output is effectively divided by the square root of the FFTSize. To maintain equal energy, output must be multiplied by the square root of the FFTSize. */
+		DividedBySqrtFFTSize,
+	};
+
+	/** FFFTSettings
+	 *
+	 * Settings for creating an IFFTAlgorithm.
+	 */
+	struct FFFTSettings
+	{
+		/** Log2Size of FFT algorithm. */
+		int32 Log2Size;
+
+		/** If true, the algorithm expects input & ouput arrays to be 128bit aligned. */
+		bool bArrays128BitAligned;
+
+		/** If true, hardware accelerated algorithms are valid. */
+		bool bEnableHardwareAcceleration;
+	};
+
+	/** IFFTAlgorithm 
+	 *
+	 * Interface for FFT algorithm.
+	 */
+	class SIGNALPROCESSING_API IFFTAlgorithm 
+	{
+		public:
+			/** virtual destructor for inheritance. */
+			virtual ~IFFTAlgorithm();
+
+			/** Number of elements in FFT */
+			virtual int32 Size() const = 0;
+
+			/** Number of floats expected in input. */
+			int32 NumInputFloats() const { return Size(); };
+
+			/** Number of floats produced in output. */
+			int32 NumOutputFloats() const { return Size() + 2; }
+
+			/** Scaling applied when performing forward FFT. */
+			virtual EFFTScaling ForwardScaling() const = 0;
+
+			/** Scaling applied when performing inverse FFT. */
+			virtual EFFTScaling InverseScaling() const = 0;
+
+			
+			/** ForwardRealToComplex
+			 *
+			 * Perform the forward FFT on real data and produce complex data.
+			 *
+			 * @param InReal - Array of floats to input into fourier transform. Must have NumInputFloats() elements.
+			 * @param OutComplex - Array of floats to store output of fourier transform. Must have NumOutputFloats() elements which represent (NumOutputFloats() / 2) complex numbers in interleaved format.
+			 */
+			virtual void ForwardRealToComplex(const float* RESTRICT InReal, float* RESTRICT OutComplex) = 0;
+
+			/** InverseComplexToReal
+			 *
+			 * Perform the inverse FFT on complex data and produce real data.
+			 *
+			 * @param InComplex - Array of floats to input into inverse fourier transform. Must have NumOutputFloats() elements which represent complex numbers in interleaved format.
+			 * @param OutReal - Array of floats to store output of inverse fourier transform. Must have NumInputFloats() elements.
+			 */
+			virtual void InverseComplexToReal(const float* RESTRICT InComplex, float* RESTRICT OutReal) = 0;
+
+			/** BatchForwardRealToComplex
+			 *
+			 * Perform the forward FFT on real data and produce complex data.
+			 *
+			 * @param InCount - The number of transforms to compute.
+			 * @param InReal - Array of array of floats to input into fourier transform. Must have InCount elements, each containing NumInputFloats() elements.
+			 * @param OutComplex - Array of array of floats to store output of fourier transform. Must have InCount elements, each with NumOutputFloats()elements which represent complex numbers in interleaved format.
+			 */
+			virtual void BatchForwardRealToComplex(int32 InCount, const float* const RESTRICT InReal[], float* RESTRICT OutComplex[]) = 0;
+
+			/** BatchInverseComplexToReal
+			 *
+			 * Perform the inverse FFT on complex data and produce real data.
+			 *
+			 * @param InCount - The number of transforms to compute.
+			 * @param InComplex - Array of array of floats to input into inverse fourier transform. Must have InCount elements, each with NumOutputFloats() elements which represent complex numbers in interleaved format.
+			 * @param OutReal - Array of array of floats to store output of inverse fourier transform. Must have InCount elements, each with NumInputFloats() elements.
+			 */
+			virtual void BatchInverseComplexToReal(int32 InCount, const float* const RESTRICT InComplex[], float* RESTRICT OutReal[]) = 0;
+	};
+
+	/** IFFTAlgorithmFactory
+	 *
+	 * Factory interface for creating fft algorithms.
+	 */
+	class SIGNALPROCESSING_API IFFTAlgorithmFactory : public IModularFeature
+	{
+		public:
+			virtual ~IFFTAlgorithmFactory()
+			{}
+
+			/** Name of modular feature for FFT factory.  */
+			static const FName GetModularFeatureName();
+
+			/** Name of this particular factory. */
+			virtual FName GetFactoryName() const = 0;
+
+			/** If true, this implementation uses hardware acceleration. */
+			virtual bool IsHardwareAccelerated() const = 0;
+
+			/** If true, this implementation requires input and output arrays to be 128 bit aligned. */
+			virtual bool Expects128BitAlignedArrays() const = 0;
+
+			/** Returns true if the input settings are supported by this factory. */
+			virtual bool AreFFTSettingsSupported(const FFFTSettings& InSettings) const = 0;
+
+			/** Creates a new FFT algorithm. */
+			virtual TUniquePtr<IFFTAlgorithm> NewFFTAlgorithm(const FFFTSettings& InSettings) = 0;
+	};
+
+	/** FFFTFactory
+	 *
+	 * FFFTFactory creates fft algorithms. It will choose hardware accelerated versions of the FFT when they are available. 
+	 */
+	class SIGNALPROCESSING_API FFFTFactory
+	{
+		public:
+			/** This denotes that no specific IFFTAlgorithmFactory is desired. */
+			static const FName AnyAlgorithmFactory;
+
+			/** NewFFTAlgorithm
+			 *
+			 * Creates and returns a new FFTAlgorithm. 
+			 *
+			 * @param InSettings - The settings used to create the FFT algorithm.
+			 * @param InAlgorithmFactoryName - If not equal to FFFTFactory::AnyAlgorithmFactory, will only uses FFT algorithm facotry if IFFTAlgorithmFactory::GetFactoryName() equals InAlgorithmFactoryName.
+			 * @return A TUniquePtr<IFFTAlgorithm> to the created FFT. This pointer can be invalid if an error occured or the fft algorithm could not be created.
+			 */
+			static TUniquePtr<IFFTAlgorithm> NewFFTAlgorithm(const FFFTSettings& InSettings, const FName& InAlgorithmFactoryName = AnyAlgorithmFactory);
+
+			/** AreFFTSettingsSupported
+			 *
+			 * Returns true if a valid FFT algorithm can be created which satifies the given settings.
+			 *
+			 * @param InSettings - Settings describing desired FFT algorithm.
+			 * @param InAlgorithmFactoryName - If not equal to FFFTFactory::AnyAlgorithmFactory, will only uses FFT algorithm facotry if IFFTAlgorithmFactory::GetFactoryName() equals InAlgorithmFactoryName.
+			 * @return True if settings are supported. Otherwise false.
+			 */
+			static bool AreFFTSettingsSupported(const FFFTSettings& InSettings, const FName& InAlgorithmFactoryName = AnyAlgorithmFactory);
+	};
+}
diff --git a/Engine/Source/Runtime/Windows/AudioMixerXAudio2/Private/AudioMixerXAudio2.cpp b/Engine/Source/Runtime/Windows/AudioMixerXAudio2/Private/AudioMixerXAudio2.cpp
--- a/Engine/Source/Runtime/Windows/AudioMixerXAudio2/Private/AudioMixerXAudio2.cpp
+++ b/Engine/Source/Runtime/Windows/AudioMixerXAudio2/Private/AudioMixerXAudio2.cpp
@@ -7,6 +7,13 @@
 class FAudioMixerModuleXAudio2 : public IAudioDeviceModule
 {
 public:
+	virtual void StartupModule() override
+	{
+		IAudioDeviceModule::StartupModule();
+
+		FModuleManager::Get().LoadModuleChecked(TEXT("AudioMixerCore"));
+	}
+
 	virtual bool IsAudioMixerModule() const override 
 	{ 
 		return true;
