Fix for build warnings from merge resolution:
1) fix UProperty in SoundSubmix.h to be an FProperty
2) Fix signed/unsigned mismatch from FMath::Abs, which returns a signed int.

diff --git a/Engine/Source/Runtime/Engine/Classes/Sound/SoundSubmix.h b/Engine/Source/Runtime/Engine/Classes/Sound/SoundSubmix.h
--- a/Engine/Source/Runtime/Engine/Classes/Sound/SoundSubmix.h
+++ b/Engine/Source/Runtime/Engine/Classes/Sound/SoundSubmix.h
@@ -90,7 +90,7 @@ protected:
 
 #if WITH_EDITOR
 	virtual void PostDuplicate(EDuplicateMode::Type DuplicateMode) override;
-	virtual void PreEditChange(UProperty* PropertyAboutToChange) override;
+	virtual void PreEditChange(FProperty* PropertyAboutToChange) override;
 	virtual void PostEditChangeProperty(struct FPropertyChangedEvent& PropertyChangedEvent) override;
 #endif // WITH_EDITOR
 	//~ End UObject Interface.
diff --git a/Engine/Source/Runtime/SignalProcessing/Private/BufferVectorOperations.cpp b/Engine/Source/Runtime/SignalProcessing/Private/BufferVectorOperations.cpp
--- a/Engine/Source/Runtime/SignalProcessing/Private/BufferVectorOperations.cpp
+++ b/Engine/Source/Runtime/SignalProcessing/Private/BufferVectorOperations.cpp
@@ -8,7 +8,7 @@ namespace Audio
 {
 	static void RestrictedPtrAliasCheck(const float* RESTRICT Ptr1, const float* RESTRICT Ptr2, uint32 NumFloatsInArray)
 	{
-		checkf(FMath::Abs(Ptr1 - Ptr2) >= NumFloatsInArray,
+		checkf(static_cast<uint32>(FMath::Abs(Ptr1 - Ptr2)) >= NumFloatsInArray,
 			TEXT("Using this function as an in-place operation will result in undefined behavior!"));
 	}
 
