static analysis warning fixes

diff --git a/Engine/Source/Runtime/SignalProcessing/Private/EQ.cpp b/Engine/Source/Runtime/SignalProcessing/Private/EQ.cpp
--- a/Engine/Source/Runtime/SignalProcessing/Private/EQ.cpp
+++ b/Engine/Source/Runtime/SignalProcessing/Private/EQ.cpp
@@ -6,110 +6,82 @@
 namespace Audio
 {
 	FEqualizer::FEqualizer()
-		: NumBands(0)
-		, NumChannels(0)
-		, FilterBands(nullptr)
+		: NumChannels(0)
 	{
 	}
 
 	FEqualizer::~FEqualizer()
 	{
-		if (FilterBands)
-		{
-			delete[] FilterBands;
-			FilterBands = nullptr;
-		}
 	}
 
 	void FEqualizer::Init(const float InSampleRate, const int32 InNumBands, const int32 InNumChannels)
 	{
-		if (FilterBands)
-		{
-			delete[] FilterBands;
-			FilterBands = nullptr;
-		}
-
-		NumBands = InNumBands;
 		NumChannels = InNumChannels;
-		FilterBands = new FBiquadFilter[NumBands];
 
-		for (int32 Band = 0; Band < NumBands; ++Band)
+		WorkBuffer.SetNum(InNumChannels);
+
+		FilterBands.SetNum(InNumBands);
+
+		for (FBiquadFilter& Filter : FilterBands)
 		{
-			FilterBands[Band].Init(InSampleRate, InNumChannels, EBiquadFilter::ParametricEQ, 500.0f, 1.0f, 0.0f);
+			Filter.Init(InSampleRate, InNumChannels, EBiquadFilter::ParametricEQ, 500.0f, 1.0f, 0.0f);
 		}
 
 	}
 
 	void FEqualizer::SetBandEnabled(const int32 InBand, const bool bInEnabled)
 	{
-		if (!FilterBands || InBand >= NumBands)
-		{
-			return;
-		}
-
 		FilterBands[InBand].SetEnabled(bInEnabled);
 	}
 
 	void FEqualizer::SetBandParams(const int32 InBand, const float InFrequency, const float InBandwidth, const float InGainDB)
 	{
-		if (!FilterBands || InBand >= NumBands)
-		{
-			return;
-		}
-
 		FilterBands[InBand].SetParams(EBiquadFilter::ParametricEQ, InFrequency, InBandwidth, InGainDB);
 	}
 
 	void FEqualizer::SetBandFrequency(const int32 InBand, const float InFrequency)
 	{
-		if (!FilterBands || InBand >= NumBands)
-		{
-			return;
-		}
-
-
 		FilterBands[InBand].SetFrequency(InFrequency);
 	}
 
 	void FEqualizer::SetBandBandwidth(const int32 InBand, const float InBandwidth)
 	{
-		if (!FilterBands || InBand >= NumBands)
-		{
-			return;
-		}
-
 		FilterBands[InBand].SetBandwidth(InBandwidth);
 	}
 
 	void FEqualizer::SetBandGainDB(const int32 InBand, const float InGainDB)
 	{
-		if (!FilterBands || InBand >= NumBands)
-		{
-			return;
-		}
-
 		FilterBands[InBand].SetGainDB(InGainDB);
 	}
 
 	void FEqualizer::ProcessAudioFrame(const float* InAudio, float* OutAudio)
 	{
-		float AudioBufferScratchInput[8];
-
-		for (int32 Channel = 0; Channel < NumChannels; ++Channel)
+		if (NumChannels < 1)
 		{
-			AudioBufferScratchInput[Channel] = InAudio[Channel];
+			return;
 		}
 
-		// Process each band in parallel
-		for (int32 Band = 0; Band < NumBands; ++Band)
+		// Copy input to output in case there are no bands to process.
+		FMemory::Memcpy(OutAudio, InAudio, sizeof(float) * NumChannels);
+
+		// Get pointers to working buffers
+		float* WorkInput = OutAudio;
+		float* WorkOutput = WorkBuffer.GetData();
+
+		for (FBiquadFilter& Filter : FilterBands)
 		{
-			FilterBands[Band].ProcessAudioFrame(AudioBufferScratchInput, OutAudio);
+			Filter.ProcessAudioFrame(WorkInput, WorkOutput);
 
-			for (int32 Channel = 0; Channel < NumChannels; ++Channel)
-			{
-				AudioBufferScratchInput[Channel] = OutAudio[Channel];
-			}
+			// Swap pointers for internal buffers
+			Swap<float*>(WorkInput, WorkOutput);
 		}
-	}
 
+		if (WorkInput != OutAudio)
+		{
+			// Note: the pointer comparison is done after a Swap
+			// If pointers do not point to same buffer, then copy  to output.
+			// Whether or not this gets called is dependent upon number of channels.
+			FMemory::Memcpy(OutAudio, WorkInput, sizeof(float) * NumChannels);
+		}
+	}
 }
diff --git a/Engine/Source/Runtime/SignalProcessing/Public/DSP/EQ.h b/Engine/Source/Runtime/SignalProcessing/Public/DSP/EQ.h
--- a/Engine/Source/Runtime/SignalProcessing/Public/DSP/EQ.h
+++ b/Engine/Source/Runtime/SignalProcessing/Public/DSP/EQ.h
@@ -42,14 +42,14 @@ namespace Audio
 
 	private:
 
-		// The number of bands in the equalizer
-		int32 NumBands;
-
 		// The number of channels in the equalizer
 		int32 NumChannels;
 
 		// The array of biquad filters
-		FBiquadFilter* FilterBands;
+		TArray<FBiquadFilter> FilterBands;
+		
+		// Temporary array for processing audio.
+		TArray<float> WorkBuffer;
 	};
 
 }
