Fix harmonic distortion in reverb with high density

diff --git a/Engine/Source/Runtime/SignalProcessing/Private/DynamicDelayAPF.cpp b/Engine/Source/Runtime/SignalProcessing/Private/DynamicDelayAPF.cpp
--- a/Engine/Source/Runtime/SignalProcessing/Private/DynamicDelayAPF.cpp
+++ b/Engine/Source/Runtime/SignalProcessing/Private/DynamicDelayAPF.cpp
@@ -6,13 +6,14 @@
 using namespace Audio;
 
 FDynamicDelayAPF::FDynamicDelayAPF(float InG, int32 InMinDelay, int32 InMaxDelay, int32 InMaxNumInternalBufferSamples, float InSampleRate)
-:	MinDelay(InMinDelay)
+:	EaseTimeInSec(1.f)
+,	MinDelay(InMinDelay)
 ,	MaxDelay(InMaxDelay)
 ,	NumDelaySamples(InMinDelay - 1)
 ,	NumInternalBufferSamples(InMaxNumInternalBufferSamples)
 {
-	G.SetValueInterrupt(InG);
 	G.Init(InSampleRate);
+	G.SetValue(InG);
 
 	checkf(NumDelaySamples >= 0, TEXT("Minimum delay must be atleast 1"));
 	// NumInternalBufferSamples must be less than the length of the delay to support buffer indexing logic.
diff --git a/Engine/Source/Runtime/SignalProcessing/Private/LateReflectionsFast.cpp b/Engine/Source/Runtime/SignalProcessing/Private/LateReflectionsFast.cpp
--- a/Engine/Source/Runtime/SignalProcessing/Private/LateReflectionsFast.cpp
+++ b/Engine/Source/Runtime/SignalProcessing/Private/LateReflectionsFast.cpp
@@ -234,7 +234,7 @@ void FLateReflectionsPlate::SetDecay(float InDecay)
 void FLateReflectionsPlate::SetDensity(float InDensity)
 {
 	Density = InDensity;
-	ModulatedAPF->SetG(-Density);
+	ModulatedAPF->SetG(FMath::Clamp(-Density, -0.9f, 0.9f));
 	APF->SetG(Density - 0.15f);
 }
 
diff --git a/Engine/Source/Runtime/SignalProcessing/Public/DSP/DynamicDelayAPF.h b/Engine/Source/Runtime/SignalProcessing/Public/DSP/DynamicDelayAPF.h
--- a/Engine/Source/Runtime/SignalProcessing/Public/DSP/DynamicDelayAPF.h
+++ b/Engine/Source/Runtime/SignalProcessing/Public/DSP/DynamicDelayAPF.h
@@ -48,7 +48,7 @@ namespace Audio
 	private:
 		// Feedback/Feedforward gain coefficient
 		FLinearEase G;
-		float EaseTimeInSec{ 2.0f };
+		float EaseTimeInSec;
 		int32 MinDelay;
 		int32 MaxDelay;
 		int32 NumDelaySamples;
