Move remaining DSP files into signal processing module

diff --git a/Engine/Source/Runtime/AudioMixer/Private/DSP/ConstantQ.cpp b/Engine/Source/Runtime/AudioMixer/Private/DSP/ConstantQ.cpp
deleted file mode 100644
diff --git a/Engine/Source/Runtime/AudioMixer/Private/DSP/InterpolatedOnePole.cpp b/Engine/Source/Runtime/AudioMixer/Private/DSP/InterpolatedOnePole.cpp
deleted file mode 100644
diff --git a/Engine/Source/Runtime/AudioMixer/Private/DSP/MelScale.cpp b/Engine/Source/Runtime/AudioMixer/Private/DSP/MelScale.cpp
deleted file mode 100644
diff --git a/Engine/Source/Runtime/AudioMixer/Private/DSP/VolumeFader.cpp b/Engine/Source/Runtime/AudioMixer/Private/DSP/VolumeFader.cpp
deleted file mode 100644
diff --git a/Engine/Source/Runtime/AudioMixer/Public/DSP/ConstantQ.h b/Engine/Source/Runtime/AudioMixer/Public/DSP/ConstantQ.h
deleted file mode 100644
diff --git a/Engine/Source/Runtime/AudioMixer/Public/DSP/InterpolatedOnePole.h b/Engine/Source/Runtime/AudioMixer/Public/DSP/InterpolatedOnePole.h
deleted file mode 100644
diff --git a/Engine/Source/Runtime/AudioMixer/Public/DSP/MelScale.h b/Engine/Source/Runtime/AudioMixer/Public/DSP/MelScale.h
deleted file mode 100644
diff --git a/Engine/Source/Runtime/AudioMixer/Public/DSP/VolumeFader.h b/Engine/Source/Runtime/AudioMixer/Public/DSP/VolumeFader.h
deleted file mode 100644
diff --git a/Engine/Source/Runtime/SignalProcessing/Private/ConstantQ.cpp b/Engine/Source/Runtime/SignalProcessing/Private/ConstantQ.cpp
new file mode 100644
--- /dev/null
+++ b/Engine/Source/Runtime/SignalProcessing/Private/ConstantQ.cpp
@@ -0,0 +1,167 @@
+// Copyright Epic Games, Inc. All Rights Reserved.
+
+#include "DSP/ConstantQ.h"
+#include "CoreMinimal.h"
+#include "DSP/FloatArrayMath.h"
+
+namespace
+{
+	float GetConstantQCenterFrequency(const int32 InBandIndex, const float InBaseFrequency, const float InBandsPerOctave)
+	{
+		check(InBandsPerOctave > 0.f);
+		return InBaseFrequency * FMath::Pow(2.f, static_cast<float>(InBandIndex) / InBandsPerOctave);
+	}
+
+	float GetConstantQBandWidth(const float InBandCenter, const float InBandsPerOctave, const float InBandWidthStretch)
+	{
+		check(InBandsPerOctave > 0.f);
+		return InBandWidthStretch * InBandCenter * (FMath::Pow(2.f, 1.f / InBandsPerOctave) - 1.f);
+	}
+
+	void FillArrayWithTruncatedGaussian(const float InCenterFreq, const float InBandWidth, const int32 InFFTSize, const float InSampleRate, TArray<float>& OutOffsetArray, int32& OutOffsetIndex)
+	{
+		check(InBandWidth > 0.f);
+		check(InFFTSize > 0);
+		check(InCenterFreq >= 0.f);
+		check(InSampleRate > 0.f);
+
+		// Determine points where gaussian will value will go below a small number
+		const float SignificantHalfBandWidth = InBandWidth * FMath::Sqrt(-2.f * FMath::Loge(SMALL_NUMBER));
+		const float LowestSignificantFreq = FMath::Clamp(InCenterFreq - SignificantHalfBandWidth, 0.f, InSampleRate / 2.f);
+		const float HighestSignificantFreq = FMath::Clamp(InCenterFreq + SignificantHalfBandWidth, 0.f, InSampleRate / 2.f);
+		int32 LowestSignificantIndex = FMath::CeilToInt(InFFTSize * LowestSignificantFreq / InSampleRate);
+		int32 HighestSignificantIndex = FMath::FloorToInt(InFFTSize * HighestSignificantFreq / InSampleRate);
+ 		int32 Num = HighestSignificantIndex - LowestSignificantIndex + 1;
+
+		if (Num < 1)
+		{
+			Num = 1;
+		}
+
+		// Prepare outputs
+		OutOffsetIndex = LowestSignificantIndex;
+		OutOffsetArray.Reset(Num);
+		OutOffsetArray.AddUninitialized(Num);
+
+		// Fill array with truncated Gaussian
+		const float BandWidthSquared = InBandWidth * InBandWidth;
+		for (int32 i = 0; i < Num; i++)
+		{
+			float FFTBinHz = static_cast<float>(i + LowestSignificantIndex) * InSampleRate / InFFTSize;
+			float DeltaHz = FFTBinHz - InCenterFreq;
+			OutOffsetArray[i] = FMath::Exp(-0.5 * (DeltaHz * DeltaHz) / BandWidthSquared);
+		}
+	}
+}
+
+namespace Audio
+{
+	TUniquePtr<FContiguousSparse2DKernelTransform> NewPseudoConstantQKernelTransform(const FPseudoConstantQKernelSettings& InSettings, const int32 InFFTSize, const float InSampleRate)
+	{
+		check(InSampleRate > 0.f);
+		check(InFFTSize > 0);
+
+		const int32 NumUsefulFFTBins = (InFFTSize / 2) + 1;
+
+		TUniquePtr<FContiguousSparse2DKernelTransform> Transform = MakeUnique<FContiguousSparse2DKernelTransform>(NumUsefulFFTBins, InSettings.NumBands);
+
+		for (int32 CQTBandIndex = 0; CQTBandIndex < InSettings.NumBands; CQTBandIndex++)
+		{
+			// Determine band center and width for this CQT band
+			const float BandCenter = GetConstantQCenterFrequency(CQTBandIndex, InSettings.KernelLowestCenterFreq, InSettings.NumBandsPerOctave);
+			const float BandWidth = GetConstantQBandWidth(BandCenter, InSettings.NumBandsPerOctave, InSettings.BandWidthStretch);
+
+			if ((BandCenter - BandWidth) > InSampleRate)
+			{
+				continue;
+			}
+
+			if (BandCenter > (2.f * InSampleRate))
+			{
+				continue;
+			}
+
+			// Create gaussian centered around center freq and with appropriate bandwidth
+			TArray<float> OffsetBandWeights;
+			int32 OffsetBandWeightsIndex = 0;
+			FillArrayWithTruncatedGaussian(BandCenter, BandWidth, InFFTSize, InSampleRate, OffsetBandWeights, OffsetBandWeightsIndex);
+
+
+			// set infs and nans to zero
+			for (int32 i = 0; i < OffsetBandWeights.Num(); i++)
+			{
+				if (!FMath::IsFinite(OffsetBandWeights[i]))
+				{
+					OffsetBandWeights[i] = 0.f;
+				}
+			}
+
+			// Need to have a lower bound for 
+			float DigitalBandWidth = FMath::Max(1.f / static_cast<float>(InFFTSize), BandWidth / InSampleRate);
+
+			// Sanity check the CQT bins to make sure we the bandwidth wasn't so small that the array is essentially empty
+			if (OffsetBandWeights.Num() > 0)
+			{
+				if (OffsetBandWeights.FilterByPredicate([](float InVal) { return InVal >= 0.5f; }).Num() < 1)
+				{
+					// Hit a condition where all values in band weights are below 0.5f. 
+					// It's a bit of an arbitrary threshold, but it tells us that the bandwidth
+					// is low enough and the FFT granularity course enough to where our pseudo 
+					// cqt windows will likely miss data. In this case we force the window to 
+					// have a single value of 1 at the nearest FFT bin.
+					FMemory::Memset(OffsetBandWeights.GetData(), 0, OffsetBandWeights.Num() * sizeof(float));
+					int32 NearestIndex = FMath::RoundToInt(InFFTSize * BandCenter / InSampleRate) - OffsetBandWeightsIndex;
+
+					if (NearestIndex < OffsetBandWeights.Num())
+					{
+						NearestIndex = FMath::Clamp(NearestIndex, 0, OffsetBandWeights.Num() - 1);
+						OffsetBandWeights[NearestIndex] = 1.f / static_cast<float>(InFFTSize);
+
+						// Need to update digital bandwidth to be bandwidth of FFTbin
+						DigitalBandWidth = 1.f / static_cast<float>(InFFTSize);
+					}
+				}
+			}
+
+
+
+
+			// Normalize window
+			float NormDenom = 1.f;
+			switch (InSettings.Normalization)
+			{
+				case EPseudoConstantQNormalization::EqualAmplitude:
+					NormDenom = 1.f;
+					break;
+				case EPseudoConstantQNormalization::EqualEuclideanNorm:
+					NormDenom = FMath::Sqrt(DigitalBandWidth * InFFTSize * FMath::Sqrt(PI  * 2.f));
+					break;
+				case EPseudoConstantQNormalization::EqualEnergy:
+				default:
+					NormDenom = DigitalBandWidth * InFFTSize * FMath::Sqrt(PI  * 2.f);
+					break;
+			}
+
+			if ((NormDenom > 0.f) && (NormDenom != 1.f))
+			{
+				ArrayMultiplyByConstantInPlace(OffsetBandWeights, 1.f / NormDenom);
+			}
+
+			// Truncate if necessary
+			if (OffsetBandWeightsIndex >= NumUsefulFFTBins)
+			{
+				OffsetBandWeightsIndex = 0;
+				OffsetBandWeights.Empty();
+			}
+			else if ((OffsetBandWeightsIndex + OffsetBandWeights.Num()) > NumUsefulFFTBins)
+			{
+				OffsetBandWeights.SetNum(NumUsefulFFTBins - OffsetBandWeightsIndex);
+			}
+
+			// Store row in transform
+			Transform->SetRow(CQTBandIndex, OffsetBandWeightsIndex, OffsetBandWeights);
+		}
+
+		return Transform;
+	}
+}
diff --git a/Engine/Source/Runtime/SignalProcessing/Private/InterpolatedOnePole.cpp b/Engine/Source/Runtime/SignalProcessing/Private/InterpolatedOnePole.cpp
new file mode 100644
--- /dev/null
+++ b/Engine/Source/Runtime/SignalProcessing/Private/InterpolatedOnePole.cpp
@@ -0,0 +1,210 @@
+// Copyright Epic Games, Inc. All Rights Reserved.
+
+#include "DSP/InterpolatedOnePole.h"
+
+namespace Audio
+{
+
+	// INTERPOLATED ONE-POLE LOW-PASS IMPLEMENTATION
+	FInterpolatedLPF::FInterpolatedLPF()
+	{
+		Z1.Init(0.0f, NumChannels);
+		Reset();
+	}
+
+
+	void FInterpolatedLPF::Init(float InSampleRate, int32 InNumChannels)
+	{
+		SampleRate = InSampleRate;
+		NumChannels = InNumChannels;
+		CutoffFrequency = -1.0f;
+		Z1.Init(0.0f, NumChannels);
+		Z1Data = Z1.GetData();
+		Reset();
+	}
+
+	void FInterpolatedLPF::StartFrequencyInterpolation(const float InTargetFrequency, const int32 InterpLength)
+	{
+		CurrInterpLength = InterpLength;
+
+		if(isFirstFrequencyChange)
+		{
+			CurrInterpLength = 0;
+			isFirstFrequencyChange = false;
+		}
+
+		if (!FMath::IsNearlyEqual(InTargetFrequency, CutoffFrequency))
+		{
+			CutoffFrequency = InTargetFrequency;
+			float NormalizedFreq = FMath::Clamp(2.0f * InTargetFrequency / SampleRate, 0.0f, 1.0f);
+			B1Target = FMath::Exp(-PI * NormalizedFreq);
+			B1Delta = (B1Target - B1Curr) / static_cast<float>(CurrInterpLength);
+		}
+
+		if (CurrInterpLength <= 1)
+		{
+			StopFrequencyInterpolation();
+		}
+	}
+
+	void FInterpolatedLPF::ProcessAudioFrame(float* RESTRICT InputFrame, float* RESTRICT OutputFrame)
+	{
+		B1Curr += B1Delta; // step forward coefficient
+
+		/*
+			[absorbing A0 coefficient]
+			-----------------------------%
+			Yn = Xn*A0 + B1*Z1;                <- old way
+			A0 = (1-B1)
+
+			Yn = Xn*(1-B1) + B1*Z1             <- (1 add, 1 sub, 2 mult)
+			Yn = Xn - B1*Xn + B1*Z1
+			Yn = Xn + B1*Z1 - B1*Xn
+			Yn = Xn + B1*(Z1 - Xn)             <- (1 add, 1 sub, 1 mult)
+		*/
+
+		for (int32 i = 0; i < NumChannels; ++i)
+		{
+			const float InputSample = InputFrame[i];
+			float Yn = InputSample + B1Curr * (Z1Data[i] - InputSample); // LPF
+			Yn = UnderflowClamp(Yn);
+			Z1Data[i] = Yn;
+			OutputFrame[i] = Yn;
+		}
+	}
+
+	void FInterpolatedLPF::ProcessAudioBuffer(float *RESTRICT InputBuffer, float *RESTRICT OutputBuffer, const int32 NumSamples)
+	{
+		for (int SampleIndex = 0; SampleIndex < NumSamples; ++SampleIndex)
+		{
+			// cache which delay term we should be using
+			const int32 ChannelIndex = SampleIndex % NumChannels;
+
+			// step forward coefficient
+			// Multiply delta by !ChannelIndex so the coefficient only accumulates at the beginning of each frame (on channel 0)
+			B1Curr += B1Delta * !ChannelIndex;
+
+			const float InputSample = InputBuffer[SampleIndex];
+			float Yn = InputSample + B1Curr * (Z1Data[ChannelIndex] - InputSample); // LPF
+			Yn = UnderflowClamp(Yn);
+			Z1Data[ChannelIndex] = Yn;
+			OutputBuffer[SampleIndex] = Yn;
+		}
+	}
+
+	void FInterpolatedLPF::Reset()
+	{
+		B1Curr = 0.0f;
+		B1Delta = 0.0f;
+		B1Target = B1Curr;
+		CurrInterpLength = 0;
+		ClearMemory();
+		Z1Data = Z1.GetData();
+		isFirstFrequencyChange = true;
+	}
+
+	void FInterpolatedLPF::ClearMemory()
+	{
+		Z1.Reset();
+		Z1.AddZeroed(NumChannels);
+	}
+
+
+	// INTERPOLATED ONE-POLE HIGH-PASS IMPLEMENTATION
+	FInterpolatedHPF::FInterpolatedHPF()
+	{
+		Z1.Init(0.0f, NumChannels);
+		Reset();
+	}
+
+	void FInterpolatedHPF::Init(float InSampleRate, int32 InNumChannels)
+	{
+		SampleRate = InSampleRate;
+		NyquistLimit = 0.5f * SampleRate - 1.0f;
+		NumChannels = InNumChannels;
+		CutoffFrequency = -1.0f;
+		Z1.Init(0.0f, NumChannels);
+		Z1Data = Z1.GetData();
+		Reset();
+	}
+
+	void FInterpolatedHPF::StartFrequencyInterpolation(const float InTargetFrequency, const int32 InterpLength)
+	{
+		CurrInterpLength = InterpLength;
+
+		if (isFirstFrequencyChange)
+		{
+			CurrInterpLength = 0;
+			isFirstFrequencyChange = false;
+		}
+
+		if (!FMath::IsNearlyEqual(InTargetFrequency, CutoffFrequency))
+		{
+			CutoffFrequency = FMath::Min(InTargetFrequency, NyquistLimit);
+
+			const float G = GetGCoefficient();
+			A0Target = G / (1.0f + G);
+
+			A0Delta = (A0Target - A0Curr) / static_cast<float>(CurrInterpLength);
+		}
+
+		if (CurrInterpLength <= 1)
+		{
+			StopFrequencyInterpolation();
+		}
+	}
+
+	void FInterpolatedHPF::ProcessAudioFrame(float* RESTRICT InputFrame, float* RESTRICT OutputFrame)
+	{
+		A0Curr += A0Delta; // step forward coefficient
+
+		for (int32 i = 0; i < NumChannels; ++i)
+		{
+			const float InputSample = InputFrame[i];
+			const float Vn = (InputSample - Z1Data[i]) * A0Curr;
+			const float LPF = Vn + Z1Data[i];
+			Z1Data[i] = Vn + LPF;
+
+			OutputFrame[i] = InputSample - LPF;
+		}
+	}
+
+	void FInterpolatedHPF::ProcessAudioBuffer(float *RESTRICT InputBuffer, float *RESTRICT OutputBuffer, const int32 NumSamples)
+	{
+		for (int SampleIndex = 0; SampleIndex < NumSamples; ++SampleIndex)
+		{
+			// cache which delay term we should be using
+			const int32 ChannelIndex = SampleIndex % NumChannels;
+
+			// step forward coefficient
+			// Multiply delta by !ChannelIndex so the coefficient only accumulates at the beginning of each frame (on channel 0)
+			A0Curr += A0Delta * !ChannelIndex;
+
+			const float InputSample = InputBuffer[SampleIndex];
+			const float Vn = (InputSample - Z1Data[ChannelIndex]) * A0Curr;
+			const float LPF = Vn + Z1Data[ChannelIndex];
+			Z1Data[ChannelIndex] = Vn + LPF;
+
+			OutputBuffer[SampleIndex] = InputSample - LPF;
+		}
+	}
+
+
+	void FInterpolatedHPF::Reset()
+	{
+		A0Curr = 0.0f;
+		A0Delta = 0.0f;
+		CurrInterpLength = 0;
+		ClearMemory();
+		Z1Data = Z1.GetData();
+		isFirstFrequencyChange = true;
+	}
+
+	void FInterpolatedHPF::ClearMemory()
+	{
+		Z1.Reset();
+		Z1.AddZeroed(NumChannels);
+	}
+
+
+} // namespace Audio
\ No newline at end of file
diff --git a/Engine/Source/Runtime/SignalProcessing/Private/MelScale.cpp b/Engine/Source/Runtime/SignalProcessing/Private/MelScale.cpp
new file mode 100644
--- /dev/null
+++ b/Engine/Source/Runtime/SignalProcessing/Private/MelScale.cpp
@@ -0,0 +1,155 @@
+// Copyright Epic Games, Inc. All Rights Reserved.
+
+#include "DSP/MelScale.h"
+#include "CoreMinimal.h"
+#include "DSP/FloatArrayMath.h"
+
+namespace
+{
+	// Fill an array with a triangle. Slopes do not need to match.
+	void FillArrayWithTriangle(const float InLowerIndex, const float InCenterIndex, const float InUpperIndex, int32& OutOffsetIndex, TArray<float>& OutArray)
+	{
+		const int32 MinIndex = FMath::Max(0, FMath::CeilToInt(InLowerIndex)); 
+		const int32 MaxIndex = FMath::FloorToInt(InUpperIndex);
+		const int32 Num = FMath::Max(0, MaxIndex - MinIndex + 1);
+
+		// Prepare output 
+		OutOffsetIndex = MinIndex;
+		OutArray.Reset(Num);
+		if (Num < 1)
+		{
+			return;
+		}
+
+		OutArray.AddUninitialized(Num);
+		FMemory::Memset(OutArray.GetData(), 0, sizeof(float) * Num);
+
+		// determin upper and lower scales for triangle
+		const float LowerScale = 1.f / FMath::Max(SMALL_NUMBER, InCenterIndex - InLowerIndex);
+		const float UpperScale = 1.f / FMath::Max(SMALL_NUMBER, InUpperIndex - InCenterIndex);
+
+		// Build traingle windows for each row
+		float* OutArrayData = OutArray.GetData();
+		for (int32 OffsetIndex = 0; OffsetIndex < Num; OffsetIndex++)
+		{
+			const float ActualIndex = static_cast<float>(MinIndex + OffsetIndex);
+			float Value = 0.0f;
+			if (ActualIndex < InCenterIndex)
+			{
+				Value = (ActualIndex - InLowerIndex) * LowerScale;
+			}
+			else
+			{
+				Value = (InUpperIndex - ActualIndex) * UpperScale;
+			}
+
+			OutArrayData[OffsetIndex] = Value < 0.0f ? 0.0f : Value;
+		}
+	}
+}
+
+namespace Audio
+{
+	float MelToHz(float InMel) {
+		return 700.f * (FMath::Exp(InMel / 1127.f) - 1.f);
+	}
+
+	float HzToMel(float InHz) {
+		return 1127.f * FMath::Loge(InHz / 700.f + 1.f);
+	}
+
+	TUniquePtr<FContiguousSparse2DKernelTransform> NewMelSpectrumKernelTransform(const FMelSpectrumKernelSettings& InSettings, const int32 InFFTSize, const float InSampleRate)
+	{
+		// Convert frequency bounds to mel
+		float MinMel = HzToMel(InSettings.KernelMinCenterFreq);
+		float MaxMel = HzToMel(InSettings.KernelMaxCenterFreq);
+		float DeltaMel = (MaxMel - MinMel) / (InSettings.NumBands - 1);
+
+		// Only need to look at half of spectrum because magnitudes are mirrored.
+		const int32 NumUsefulFrequencyBins = InFFTSize / 2 + 1;
+		float BinsPerHz = (float)InFFTSize / (float)InSampleRate;
+		float HzPerBin = (float)InSampleRate / (float)InFFTSize;
+
+		// Used to define bounds of mel band windows.
+		float PreviousMelCenter = MinMel - DeltaMel;
+		float CurrentMelCenter = MinMel;
+		float NextMelCenter = MinMel + DeltaMel;
+
+		// Mel band windows are contiguous and sparse, so use this transform to make FFT->Mel fast.
+		TUniquePtr<FContiguousSparse2DKernelTransform> KernelTransform = MakeUnique<FContiguousSparse2DKernelTransform>(NumUsefulFrequencyBins, InSettings.NumBands);
+
+		// Create mel band windows one by one.
+		for (int32 MelBandIndex = 0; MelBandIndex < InSettings.NumBands; MelBandIndex++)
+		{
+			// Convert Mels to Hz
+			float LowerBandFreq = MelToHz(PreviousMelCenter);
+			float UpperBandFreq = MelToHz(NextMelCenter);
+			float CenterBandFreq = MelToHz(CurrentMelCenter);
+
+			//  Stretch frequency bandwidth if necessary
+			if (InSettings.BandWidthStretch != 1.0f)
+			{
+				LowerBandFreq = CenterBandFreq + InSettings.BandWidthStretch * (LowerBandFreq - CenterBandFreq);
+				UpperBandFreq = CenterBandFreq + InSettings.BandWidthStretch * (UpperBandFreq - CenterBandFreq);
+			}
+
+			// Determine bounding bins
+			float LowerFFTBin = LowerBandFreq * BinsPerHz;
+			float UpperFFTBin = UpperBandFreq * BinsPerHz;
+			float CenterFFTBin = CenterBandFreq * BinsPerHz;
+
+			// Create a triangle in frequency domain to account for mel band.
+			TArray<float> RowOffsetValues;
+			int32 RowOffsetIndex = 0;
+			FillArrayWithTriangle(LowerFFTBin, CenterFFTBin, UpperFFTBin, RowOffsetIndex, RowOffsetValues);
+
+			// By default, FillArrayWithTriangle creates constant amplitude triangle. So only noralize 
+			// if constant energy is chosen.
+			float NormDenom = 1.f;
+			switch (InSettings.Normalization)
+			{
+				case EMelNormalization::EqualAmplitude:
+					// By default, FillArrayWithTriangle creates constant amplitude triangle. So only noralize 
+					// if constant energy is chosen.
+					NormDenom = 1.f;
+					break;
+
+				case EMelNormalization::EqualEuclideanNorm:
+					NormDenom = FMath::Sqrt(0.5f * (UpperFFTBin - LowerFFTBin));
+					break;
+
+				case EMelNormalization::EqualEnergy:
+				default:
+					// weight bands to have constant energy of 1
+					NormDenom = 0.5f * (UpperFFTBin - LowerFFTBin);
+					break;
+			}
+
+			if ((NormDenom > 0.f) && (NormDenom != 1.f))
+			{
+				ArrayMultiplyByConstantInPlace(RowOffsetValues, 1.f / NormDenom);
+			}
+
+			// Check bounds of data being input into kernel
+			int32 MaxNumRowOFfsetValues = FMath::Max(0, NumUsefulFrequencyBins - RowOffsetIndex);
+			if (RowOffsetValues.Num() > MaxNumRowOFfsetValues)
+			{
+				RowOffsetValues.SetNum(MaxNumRowOFfsetValues);
+			}
+			if (RowOffsetIndex > NumUsefulFrequencyBins)
+			{
+				RowOffsetIndex = 0;
+			}
+
+			// Add row to kernel
+			KernelTransform->SetRow(MelBandIndex, RowOffsetIndex, TArrayView<const float>(RowOffsetValues));
+
+			// increment mel band counters 
+			PreviousMelCenter = CurrentMelCenter;
+			CurrentMelCenter = NextMelCenter;
+			NextMelCenter += DeltaMel;
+		}
+
+		return KernelTransform;
+	}
+}
diff --git a/Engine/Source/Runtime/SignalProcessing/Private/VolumeFader.cpp b/Engine/Source/Runtime/SignalProcessing/Private/VolumeFader.cpp
new file mode 100644
--- /dev/null
+++ b/Engine/Source/Runtime/SignalProcessing/Private/VolumeFader.cpp
@@ -0,0 +1,274 @@
+// Copyright Epic Games, Inc. All Rights Reserved.
+#include "DSP/VolumeFader.h"
+
+#include "DSP/Dsp.h"
+
+
+namespace Audio
+{
+	FVolumeFader::FVolumeFader()
+		: Alpha(1.0f)
+		, Target(1.0f)
+		, ActiveDuration(-1.0f)
+		, FadeDuration(-1.0f)
+		, Elapsed(0.0f)
+		, FadeCurve(EFaderCurve::Linear)
+	{
+	}
+
+	float FVolumeFader::AlphaToVolume(float InAlpha, EFaderCurve InCurve)
+	{
+		switch (InCurve)
+		{
+		case EFaderCurve::Linear:
+		{
+			return InAlpha;
+		}
+		break;
+
+		case EFaderCurve::SCurve:
+		{
+			float Volume = 0.5f * Audio::FastSin(PI * InAlpha - HALF_PI) + 0.5f;
+			return FMath::Max(0.0f, Volume);
+		}
+		break;
+
+		case EFaderCurve::Sin:
+		{
+			float Volume = Audio::FastSin(HALF_PI * InAlpha);
+			return FMath::Max(0.0f, Volume);
+		}
+		break;
+
+		case EFaderCurve::Logarithmic:
+		{
+			return Audio::ConvertToLinear(InAlpha);
+		}
+		break;
+
+		default:
+		{
+			static_assert(static_cast<int32>(EFaderCurve::Count) == 4, "Possible missing switch case coverage for EAudioFade");
+		}
+		break;
+		}
+
+		return 1.0f;
+	}
+
+	void FVolumeFader::Activate()
+	{
+		if (ActiveDuration == 0.0f)
+		{
+			ActiveDuration = -1.0f;
+		}
+	}
+
+	void FVolumeFader::Deactivate()
+	{
+		StopFade();
+		ActiveDuration = 0.0f;
+	}
+
+	float FVolumeFader::GetActiveDuration() const
+	{
+		return ActiveDuration;
+	}
+
+	float FVolumeFader::GetVolume() const
+	{
+		return AlphaToVolume(Alpha, FadeCurve);
+	}
+
+	float FVolumeFader::GetVolumeAfterTime(float InDeltaTime) const
+	{
+		InDeltaTime = FMath::Max(0.0f, InDeltaTime);
+
+		// Keep stepping towards our target until we hit our stop time & Clamp
+		float FutureAlpha = Alpha;
+		const float Duration = Elapsed + InDeltaTime;
+
+		// If set to deactivate first, return default volume value of 1.0f
+		if (ActiveDuration >= 0.0f && Duration > ActiveDuration)
+		{
+			return 1.0f;
+		}
+
+		if (Duration < FadeDuration)
+		{
+			// Choose min/max bound and clamp dt to prevent unwanted spikes in volume
+			float MinValue = 0.0f;
+			float MaxValue = 0.0f;
+			if (Alpha < Target)
+			{
+				MinValue = Alpha;
+				MaxValue = Target;
+			}
+			else
+			{
+				MinValue = Target;
+				MaxValue = Alpha;
+			}
+
+			FutureAlpha = Alpha + ((Target - Alpha) * InDeltaTime / (FadeDuration - Elapsed));
+			FutureAlpha = FMath::Clamp(FutureAlpha, MinValue, MaxValue);
+		}
+
+		return AlphaToVolume(FutureAlpha, FadeCurve);
+	}
+
+	float FVolumeFader::GetFadeDuration() const
+	{
+		return FadeDuration;
+	}
+
+	EFaderCurve FVolumeFader::GetCurve() const
+	{
+		return FadeCurve;
+	}
+
+	float FVolumeFader::GetTargetVolume() const
+	{
+		switch (FadeCurve)
+		{
+		case EFaderCurve::Linear:
+		case EFaderCurve::SCurve:
+		case EFaderCurve::Sin:
+		{
+			return Target;
+		}
+		break;
+
+		case EFaderCurve::Logarithmic:
+		{
+			return Audio::ConvertToLinear(Target);
+		}
+		break;
+
+		default:
+		{
+			static_assert(static_cast<int32>(EFaderCurve::Count) == 4, "Possible missing switch case coverage for EAudioFade");
+		}
+		break;
+		}
+
+		return 1.0f;
+	}
+
+	bool FVolumeFader::IsActive() const
+	{
+		return Elapsed < ActiveDuration || ActiveDuration < 0.0f;
+	}
+
+	bool FVolumeFader::IsFading() const
+	{
+		return Elapsed < FadeDuration;
+	}
+
+	bool FVolumeFader::IsFadingIn() const
+	{
+		return IsFading() && Target > Alpha;
+	}
+
+	bool FVolumeFader::IsFadingOut() const
+	{
+		return IsFading() && Target < Alpha;
+	}
+
+	void FVolumeFader::SetActiveDuration(float InDuration)
+	{
+		if (InDuration < 0.0f)
+		{
+			Activate();
+		}
+		else
+		{
+			ActiveDuration = InDuration;
+		}
+	}
+
+	void FVolumeFader::SetVolume(float InVolume)
+	{
+		Alpha = InVolume;
+		Elapsed = 0.0f;
+		FadeCurve = EFaderCurve::Linear;
+		FadeDuration = -1.0f;
+		Target = InVolume;
+	}
+
+	void FVolumeFader::StartFade(float InVolume, float InDuration, EFaderCurve InCurve)
+	{
+		if (InDuration <= 0.0f)
+		{
+			SetVolume(InVolume);
+			return;
+		}
+
+		if (InCurve != EFaderCurve::Logarithmic)
+		{
+			if (FadeCurve == EFaderCurve::Logarithmic)
+			{
+				Alpha = Audio::ConvertToLinear(Alpha);
+			}
+			Target = InVolume;
+		}
+		else
+		{
+			const float DecibelFloor = KINDA_SMALL_NUMBER; // -80dB
+			if (FadeCurve != EFaderCurve::Logarithmic)
+			{
+				Alpha = Audio::ConvertToDecibels(Alpha, DecibelFloor);
+			}
+			Target = Audio::ConvertToDecibels(InVolume, DecibelFloor);
+		}
+
+		Elapsed = 0.0f;
+		FadeCurve = InCurve;
+		FadeDuration = InDuration;
+	}
+
+	void FVolumeFader::StopFade()
+	{
+		if (FadeCurve == EFaderCurve::Logarithmic)
+		{
+			Alpha = Audio::ConvertToLinear(Alpha);
+		}
+		Target = Alpha;
+		FadeCurve = EFaderCurve::Linear;
+		FadeDuration = -1.0f;
+	}
+
+	void FVolumeFader::Update(float InDeltaTime)
+	{
+		Elapsed += InDeltaTime;
+
+		if (!IsFading() || !IsActive())
+		{
+			return;
+		}
+
+		// Keep stepping towards target & clamp until fade duration has expired
+		// Choose min/max bound and clamp dt to prevent unwanted spikes in volume
+		float MinValue = 0.0f;
+		float MaxValue = 0.0f;
+		if (Alpha < Target)
+		{
+			MinValue = Alpha;
+			MaxValue = Target;
+		}
+		else
+		{
+			MinValue = Target;
+			MaxValue = Alpha;
+		}
+
+		Alpha += (Target - Alpha) * InDeltaTime / (FadeDuration - Elapsed);
+		Alpha = FMath::Clamp(Alpha, MinValue, MaxValue);
+
+		// Optimization that avoids fade calculations once alpha reaches target
+		if (InDeltaTime > SMALL_NUMBER && FMath::IsNearlyEqual(Alpha, Target))
+		{
+			StopFade();
+		}
+	}
+} // namespace Audio
\ No newline at end of file
diff --git a/Engine/Source/Runtime/SignalProcessing/Public/DSP/ConstantQ.h b/Engine/Source/Runtime/SignalProcessing/Public/DSP/ConstantQ.h
new file mode 100644
--- /dev/null
+++ b/Engine/Source/Runtime/SignalProcessing/Public/DSP/ConstantQ.h
@@ -0,0 +1,55 @@
+// Copyright Epic Games, Inc. All Rights Reserved.
+
+#pragma once
+
+#include "CoreMinimal.h"
+#include "DSP/FloatArrayMath.h"
+
+namespace Audio
+{
+	/** Band normalization schemes */
+	enum class EPseudoConstantQNormalization : uint8
+	{
+		EqualAmplitude,		//< No energy scaling. All bands weighted so max is 1.
+		EqualEuclideanNorm, //< Scale energy by euclidean norm. Good when using magnitude spectrum.
+		EqualEnergy			//< Scale energy by total energy. Good when using power spectrum.
+	};
+
+	/** Settings for Pseudo Constant Q Kernel generation. */
+	struct SIGNALPROCESSING_API FPseudoConstantQKernelSettings
+	{
+		int32 NumBands;									//< Total number of resulting constant Q bands.
+		float NumBandsPerOctave;							//< Number of bands to space within an octave.
+		float KernelLowestCenterFreq;					//< The starting frequency of the first band.
+		float BandWidthStretch;							//< Stretching factor to control overlap of adjacent bands.
+		EPseudoConstantQNormalization Normalization;	//< Normalization scheme for bands.
+
+		/** Default settings for Pseudo Constant Q Kernel */
+		FPseudoConstantQKernelSettings()
+		:	NumBands(96)
+		,	NumBandsPerOctave(12)
+		,	KernelLowestCenterFreq(40.f)
+		,	BandWidthStretch(1.0) 
+		,	Normalization(EPseudoConstantQNormalization::EqualEnergy)
+		{}
+	};
+
+	/** NewPseudoConstantQKernelTransform
+	 *
+	 *  This creates a new Pseudo Constant Q Kernel Transform. Pseudo Constant Q differs from standard Constant Q in
+	 *  that it applies a window to an existing DFT output as opposed to using a bank of filters.  The use of
+	 *  an FFT speeds up the calculation but introduces some limitations in bandwidth due to the granularity of the
+	 *  DFT. 
+	 *
+	 *  The resulting kernel will transform an array with (InFFTSize / 2 + 1) elements into an array with 
+	 *  InSettings.NumBands elements.
+	 *
+	 *  InSettings controls the number of bands and band properties.
+	 *  InFFTSize controls the expected input size to the kernel.
+	 *  InSampleRate describes the sampling rate of the analyzed audio.
+	 *
+	 *  Returns a TUniquePtr to a FContiguousSparse2DKernelTransform.
+	 */
+	SIGNALPROCESSING_API TUniquePtr<FContiguousSparse2DKernelTransform> NewPseudoConstantQKernelTransform(const FPseudoConstantQKernelSettings& InSettings, const int32 InFFTSize, const float InSampleRate);
+
+}
diff --git a/Engine/Source/Runtime/SignalProcessing/Public/DSP/InterpolatedOnePole.h b/Engine/Source/Runtime/SignalProcessing/Public/DSP/InterpolatedOnePole.h
new file mode 100644
--- /dev/null
+++ b/Engine/Source/Runtime/SignalProcessing/Public/DSP/InterpolatedOnePole.h
@@ -0,0 +1,146 @@
+// Copyright Epic Games, Inc. All Rights Reserved.
+
+#pragma once
+#include "CoreMinimal.h"
+#include "DSP/Dsp.h"
+
+
+namespace Audio
+{
+	class SIGNALPROCESSING_API FInterpolatedLPF
+	{
+	public:
+		// ctor
+		FInterpolatedLPF();
+
+		FORCEINLINE float GetCutoffFrequency() const
+		{
+			return CutoffFrequency;
+		}
+
+		void Init(float InSampleRate, int32 InNumChannels);
+
+		/*
+			StartFrequencyInterpolation():
+				Intended use: call once at the beginning of callback, pass the callback length
+				then call ProcessAudioFrame() exactly "InterpLength" number of times
+				At the end of the callback, call StopFrequencyInterpolation()
+
+				Failing to call StopFrequencyInterpolation() after more than InterpLength calls to ProcessAudioFrame
+				will cause the filter to become unstable and blow up to +/- inf!
+		*/
+		void StartFrequencyInterpolation(const float InTargetFrequency, const int32 InterpLength = 1);
+
+		// interpolates coefficient and processes a sample
+		void ProcessAudioFrame(float* RESTRICT InputFrame, float* RESTRICT OutputFrame);
+
+		// interpolates coefficient and processes a buffer
+		void ProcessAudioBuffer(float* RESTRICT InputBuffer, float* RESTRICT OutputBuffer, const int32 NumSamples);
+
+		/*
+			StopFrequencyInterpolation() needs to be called manually when the interpolation should be done.
+			Snaps the coefficient to the target value.
+			Setting the lerp delta to 0.0f avoids introducing a branch in ProcessAudioFrame()
+		*/
+		FORCEINLINE void StopFrequencyInterpolation()
+		{
+			B1Delta = 0.0f;
+			B1Curr = B1Target;
+		}
+
+		// Resets the sample delay to 0
+		void Reset();
+
+		// Clears memory without reevaluating coefficients.
+		// This function is useful when there is a break between ProcessAudio calls.
+		void ClearMemory();
+
+	private:
+		float CutoffFrequency{ -1.0f };
+		float B1Curr{ 0.0f }; // coefficient
+		float B1Delta{ 0.0f }; // coefficient step size 
+		float B1Target{ 0.0f };
+		TArray<float> Z1; // multi-channel delay terms
+		float* Z1Data{ nullptr };
+		int32 CurrInterpLength{ 0 };
+		int32 NumInterpSteps;
+		float SampleRate{ 0 };
+		int32 NumChannels{ 1 };
+		int32 isFirstFrequencyChange : 1;
+
+	}; // class Interpolated Low pass filter
+
+
+
+	class SIGNALPROCESSING_API FInterpolatedHPF
+	{
+	public:
+		// ctor
+		FInterpolatedHPF();
+
+		FORCEINLINE float GetCutoffFrequency() const
+		{
+			return CutoffFrequency;
+		}
+
+		void Init(float InSampleRate, int32 InNumChannels);
+
+		/*
+			StartFrequencyInterpolation():
+				Intended use: call once at the beginning of callback, pass the callback length
+				then call ProcessAudioFrame() exactly "InterpLength" number of times
+				At the end of the callback, call StopFrequencyInterpolation()
+
+				Failing to call StopFrequencyInterpolation() after more than InterpLength calls to ProcessAudioFrame
+				will cause the filter to become unstable and blow up to +/- inf!
+		*/
+		void StartFrequencyInterpolation(const float InTargetFrequency, const int32 InterpLength = 1);
+
+		// interpolates coefficient and processes a sample
+		void ProcessAudioFrame(float* RESTRICT InputFrame, float* RESTRICT OutputFrame);
+
+		// interpolates coefficient and processes a buffer
+		void ProcessAudioBuffer(float* RESTRICT InputBuffer, float* RESTRICT OutputBuffer, const int32 NumSamples);
+
+		/*
+			StopFrequencyInterpolation() needs to be called manually when the interpolation should be done.
+			Snaps the coefficient to the target value.
+			Setting the lerp delta to 0.0f avoids introducing a branch in ProcessAudioFrame()
+		*/
+		FORCEINLINE void StopFrequencyInterpolation()
+		{
+			A0Delta = 0.0f;
+			A0Curr = A0Target;
+		}
+
+		// Resets the sample delay to 0
+		void Reset();
+
+		// Clears memory without reevaluating coefficients.
+		// This function is useful when there is a break between ProcessAudio calls.
+		void ClearMemory();
+		
+	private:
+		FORCEINLINE float GetGCoefficient() const
+		{
+			const float OmegaDigital = 2.0f * PI * CutoffFrequency;
+			const float OmegaAnalog = 2.0f * SampleRate * Audio::FastTan(0.5f * OmegaDigital / SampleRate);
+			const float G = 0.5f * OmegaAnalog / SampleRate;
+			return G;
+		}
+
+		float CutoffFrequency{ -1.0f };
+		float A0Curr{ 0.0f }; // coefficient
+		float A0Delta{ 0.0f }; // coefficient step size
+		float A0Target{ 0.0f };
+		TArray<float> Z1; // multi-channel delay terms
+		float* Z1Data{ nullptr };
+		int32 CurrInterpLength{ 0 };
+		int32 NumInterpSteps;
+		float SampleRate{ 0 };
+		float NyquistLimit{ 0.0f };
+		int32 NumChannels{ 1 };
+		int32 isFirstFrequencyChange : 1;
+	
+	}; // class Interpolated High pass filter
+} // namespace Audio
diff --git a/Engine/Source/Runtime/SignalProcessing/Public/DSP/MelScale.h b/Engine/Source/Runtime/SignalProcessing/Public/DSP/MelScale.h
new file mode 100644
--- /dev/null
+++ b/Engine/Source/Runtime/SignalProcessing/Public/DSP/MelScale.h
@@ -0,0 +1,44 @@
+// Copyright Epic Games, Inc. All Rights Reserved.
+
+#pragma once
+#include "CoreMinimal.h"
+#include "DSP/FloatArrayMath.h"
+
+
+namespace Audio
+{
+	/** Convert a mel frequency to a frequency in hz. */
+	SIGNALPROCESSING_API float MelToHz(float InMel);
+
+	/** Convert  a frequency in hz to a mel frequency. */
+	SIGNALPROCESSING_API float HzToMel(float InHz);
+
+	/** Normalization methods for a mel transform. */
+	enum class EMelNormalization : uint8
+	{
+		EqualAmplitude, 	//< Useful when working on magnitude spectrum
+		EqualEuclideanNorm, //< Scale energy by euclidean norm. Good when using magnitude spectrum.
+		EqualEnergy			//< Useful when working on power spectrum
+	};
+
+	/** Settings for a mel kernel which transforms an linearly space spectrum (e.g. FFT Magnitude)
+	 * to a mel spectrum */
+	struct SIGNALPROCESSING_API FMelSpectrumKernelSettings
+	{
+		int32 NumBands;						//< Number of bands in Mel spectrum
+		float KernelMinCenterFreq;			//< Minimum frequency of lowest mel band.
+		float KernelMaxCenterFreq;			//< Maximum frequency of highest mel band.
+		EMelNormalization Normalization;	//< Mel band scaling. EqualAmplitude is better for the magnitude spectrum, while EqualEnergy is better for the power spectrum.
+		float BandWidthStretch;				//< Controls the wideness of each band.
+
+		FMelSpectrumKernelSettings()
+		:	NumBands(26)
+		,	KernelMinCenterFreq(100.f)
+		,	KernelMaxCenterFreq(10000.f)
+		,	Normalization(EMelNormalization::EqualEnergy)
+		,	BandWidthStretch(1.f)
+		{}
+	};
+
+	SIGNALPROCESSING_API TUniquePtr<FContiguousSparse2DKernelTransform> NewMelSpectrumKernelTransform(const FMelSpectrumKernelSettings& InSettings, const int32 InFFTSize, const float InSampleRate);
+}
diff --git a/Engine/Source/Runtime/SignalProcessing/Public/DSP/VolumeFader.h b/Engine/Source/Runtime/SignalProcessing/Public/DSP/VolumeFader.h
new file mode 100644
--- /dev/null
+++ b/Engine/Source/Runtime/SignalProcessing/Public/DSP/VolumeFader.h
@@ -0,0 +1,155 @@
+// Copyright Epic Games, Inc. All Rights Reserved.
+#pragma once
+
+#include "CoreMinimal.h"
+
+
+namespace Audio
+{
+	/** Type of fade to use when adjusting a fader's volume over time. */
+	enum class EFaderCurve : uint8
+	{
+		// Linear Fade
+		Linear,
+
+		// Logarithmic Fade
+		Logarithmic,
+
+		// S-Curve, Sinusoidal Fade
+		SCurve,
+
+		// Equal Power, Sinusoidal Fade
+		Sin,
+
+		Count
+	};
+
+
+	/** Control-rate fader for managing volume fades of various standard shapes. */
+	class SIGNALPROCESSING_API FVolumeFader
+	{
+	public:
+		FVolumeFader();
+
+		/**
+		 * Activates the fader if currently deactivated. Sets duration to indefinite (-1.0f)
+		 */
+		void Activate();
+
+		/**
+		 * Deactivates the fader, causing it to remain
+		 * at the current value and disregard update.
+		 */
+		void Deactivate();
+
+		/**
+		 * Gets time fader is to remain active (-1.0f if always active).
+		 */
+		float GetActiveDuration() const;
+
+		/**
+		 * Returns current volume of fader
+		 */
+		float GetVolume() const;
+
+		/**
+		 * Returns the volume given the delta from the current time
+		 * into the future (Effectively like running to retrieve volume
+		 * but without actually updating internal state).
+		 */
+		float GetVolumeAfterTime(float InDeltaTime) const;
+
+		/**
+		 * Returns the duration of the fade.
+		 */
+		float GetFadeDuration() const;
+
+		/**
+		 * Returns the curve type of the fader
+		 */
+		EFaderCurve GetCurve() const;
+
+		/**
+		 * Returns the target volume of the fader
+		 */
+		float GetTargetVolume() const;
+
+		/**
+		 * Whether or not the fader is active.
+		 */
+		bool IsActive() const;
+
+		/**
+		 * Returns whether or not the fader is currently
+		 * fading over time.
+		 */
+		bool IsFading() const;
+
+		/**
+		 * Returns whether or not the fader is currently
+		 * fading over time and value is increasing.
+		 */
+		bool IsFadingIn() const;
+
+		/**
+		 * Returns whether or not the fader is currently
+		 * fading over time and value is decreasing.
+		 */
+		bool IsFadingOut() const;
+
+		/**
+		 * Sets the duration the fader is to be active in the future,
+		 * after which point the fader is disabled.  When disabled,
+		 * fader will hold the current fade state until activated,
+		 * at which point it will continue applying the fade over the
+		 * remainder of the fade duration.
+		 */
+		void SetActiveDuration(float InDuration);
+
+		/**
+		 * Sets the volume immediately, interrupting any currently active fade.
+		 */
+		void SetVolume(float InVolume);
+
+		/**
+		 * Applies a volume fade over time with the provided parameters.
+		 */
+		void StartFade(float InVolume, float InDuration, EFaderCurve InCurve);
+
+		/**
+		 * Stops fade, maintaining the current value as the target.
+		 */
+		void StopFade();
+
+		/**
+		 * Updates the fader's state with the given delta in time since last update.
+		 */
+		void Update(float InDeltaTime);
+
+	private:
+		/** Converts value to final resulting volume */
+		static float AlphaToVolume(float InAlpha, EFaderCurve InCurve);
+
+		/** Current value used to linear interpolate over update delta
+		  * (Normalized value for non-log, -80dB to 0dB for log)
+		  */
+		float Alpha;
+
+		/** Target value used to linear interpolate over update delta
+		  * (Normalized value for non-log, -80dB to 0dB for log)
+		  */
+		float Target;
+
+		/** Duration fader is to be active */
+		float ActiveDuration;
+
+		/** Duration fader is to perform fade */
+		float FadeDuration;
+
+		/** Time elapsed since fade has been initiated */
+		float Elapsed;
+
+		/** Audio fader curve to use */
+		EFaderCurve FadeCurve;
+	};
+} // namespace Audio
