Updating scaling on FFTs

diff --git a/Engine/Source/Runtime/SignalProcessing/Private/AudioFFT.cpp b/Engine/Source/Runtime/SignalProcessing/Private/AudioFFT.cpp
--- a/Engine/Source/Runtime/SignalProcessing/Private/AudioFFT.cpp
+++ b/Engine/Source/Runtime/SignalProcessing/Private/AudioFFT.cpp
@@ -537,10 +537,10 @@ namespace Audio
 			virtual int32 Size() const override { return FFTSize; }
 
 			// Scaling applied when performing forward FFT.
-			virtual EFFTScaling ForwardScaling() const { return EFFTScaling::None; }
+			virtual EFFTScaling ForwardScaling() const { return EFFTScaling::MultipliedBySqrtFFTSize; }
 
 			/* Scaling applied when performing inverse FFT. */
-			virtual EFFTScaling InverseScaling() const { return EFFTScaling::None; }
+			virtual EFFTScaling InverseScaling() const { return EFFTScaling::DividedBySqrtFFTSize; }
 
 			
 			// ForwardRealToComplex
diff --git a/Engine/Source/Runtime/SignalProcessing/Private/VectorFFT.cpp b/Engine/Source/Runtime/SignalProcessing/Private/VectorFFT.cpp
--- a/Engine/Source/Runtime/SignalProcessing/Private/VectorFFT.cpp
+++ b/Engine/Source/Runtime/SignalProcessing/Private/VectorFFT.cpp
@@ -1141,13 +1141,13 @@ namespace Audio
 	/** Scaling applied when performing forward FFT. */
 	EFFTScaling FVectorRealToComplexFFT::ForwardScaling() const
 	{
-		return EFFTScaling::None;
+		return EFFTScaling::MultipliedBySqrtFFTSize;
 	}
 
 	/** Scaling applied when performing inverse FFT. */
 	EFFTScaling FVectorRealToComplexFFT::InverseScaling() const
 	{
-		return EFFTScaling::None;
+		return EFFTScaling::DividedBySqrtFFTSize;
 	}
 
 	void FVectorRealToComplexFFT::ForwardRealToComplex(const float* RESTRICT InReal, float* RESTRICT OutComplex)
