fixing dsp bugs in audio spectrum data interface

diff --git a/Engine/Plugins/FX/Niagara/Source/Niagara/Classes/NiagaraDataInterfaceAudioSpectrum.h b/Engine/Plugins/FX/Niagara/Source/Niagara/Classes/NiagaraDataInterfaceAudioSpectrum.h
--- a/Engine/Plugins/FX/Niagara/Source/Niagara/Classes/NiagaraDataInterfaceAudioSpectrum.h
+++ b/Engine/Plugins/FX/Niagara/Source/Niagara/Classes/NiagaraDataInterfaceAudioSpectrum.h
@@ -7,6 +7,7 @@
 #include "DSP/ConstantQ.h"
 #include "DSP/FFTAlgorithm.h"
 #include "DSP/FloatArrayMath.h"
+#include "DSP/AudioFFT.h"
 #include "NiagaraDataInterfaceAudio.h"
 #include "NiagaraDataInterfaceAudioSpectrum.generated.h"
 
@@ -77,7 +78,7 @@ private:
 	void ResizeCQT(float InMinimumFrequency, float InMaximumFrequency, int32 InNumBands);
 	void ResizeSpectrumBuffer(int32 InNumChannels, int32 InNumBands);
 	void ResizeAudioTransform(float InMinimumFrequency, float InMaximumFrequency, float InSampleRate, int32 InNumBands);
-	void ResizeSlidingBuffer(int32 InNumChannels, int32 InFFTSize);
+	void ResizeWindow(int32 InNumChannels, int32 InFFTSize);
 
 
 	// Get settings 
@@ -107,12 +108,13 @@ private:
 	TUniquePtr<Audio::FContiguousSparse2DKernelTransform> CQTKernel;
 	TUniquePtr<Audio::IFFTAlgorithm> FFTAlgorithm;
 
-	Audio::AlignedFloatBuffer InterleavedWindowBuffer;
-	Audio::AlignedFloatBuffer WindowBuffer;
+	Audio::AlignedFloatBuffer InterleavedBuffer;
+	Audio::AlignedFloatBuffer DeinterleavedBuffer;
 	Audio::AlignedFloatBuffer FFTInputBuffer;
 	Audio::AlignedFloatBuffer FFTOutputBuffer;
 	Audio::AlignedFloatBuffer PowerSpectrumBuffer;
 	Audio::AlignedFloatBuffer SpectrumBuffer;
+	Audio::AlignedFloatBuffer WindowBuffer;
 
 	// Handle for the SRV used by the generated HLSL.
 	FReadBuffer GPUBuffer;
diff --git a/Engine/Plugins/FX/Niagara/Source/Niagara/Private/NiagaraDataInterfaceAudioSpectrum.cpp b/Engine/Plugins/FX/Niagara/Source/Niagara/Private/NiagaraDataInterfaceAudioSpectrum.cpp
--- a/Engine/Plugins/FX/Niagara/Source/Niagara/Private/NiagaraDataInterfaceAudioSpectrum.cpp
+++ b/Engine/Plugins/FX/Niagara/Source/Niagara/Private/NiagaraDataInterfaceAudioSpectrum.cpp
@@ -112,7 +112,7 @@ void FNiagaraDataInterfaceProxySpectrum::ResizeCQT(float InMinimumFrequency, flo
 
 		if (NewFFTSize != OriginalFFTSize)
 		{
-			ResizeSlidingBuffer(NumChannels, NewFFTSize);
+			ResizeWindow(NumChannels, NewFFTSize);
 		}
 
 	}
@@ -179,6 +179,7 @@ void FNiagaraDataInterfaceProxySpectrum::PostDataToGPU()
 
 			for (int32 ChannelIndex = 0; ChannelIndex < NumChannels; ChannelIndex++)
 			{
+				
 				int32 Pos = ChannelIndex * NumBands;
 				FPlatformMemory::Memcpy(&BufferData[Pos], ChannelSpectrumBuffers[ChannelIndex].GetData(), NumBytesInChannelBuffer);
 			}
@@ -222,16 +223,17 @@ void FNiagaraDataInterfaceProxySpectrum::UpdateSpectrum()
 	// Grab audio from buffer.
 	int32 NumSamplesToPop = GetNumSamplesAvailable();
 	NumSamplesToPop -= (NumSamplesToPop % NumChannels);
-
+	PopBuffer.Reset();
+	
 	if (NumSamplesToPop > 0)
 	{
-		PopBuffer.Reset();
+		
 		PopBuffer.AddUninitialized(NumSamplesToPop);
 
 		bool bUseLatestAudio = false;
-		PopAudio(PopBuffer.GetData(), NumSamplesToPop, bUseLatestAudio);
+		int32 NumPopped = PopAudio(PopBuffer.GetData(), NumSamplesToPop, bUseLatestAudio);
 	}
-
+	
 	// If we're in a bad internal state, give up here.	
 	if (!FFTAlgorithm.IsValid() || !CQTKernel.IsValid())
 	{
@@ -239,7 +241,7 @@ void FNiagaraDataInterfaceProxySpectrum::UpdateSpectrum()
 	}
 
 	// Run sliding window over available audio.
-	FSlidingWindow SlidingWindow(*SlidingBuffer, PopBuffer, InterleavedWindowBuffer);
+	FSlidingWindow SlidingWindow(*SlidingBuffer, PopBuffer, InterleavedBuffer);
 
 	int32 NumWindows = 0;
 	for (Audio::AlignedFloatBuffer& InterleavedWindow : SlidingWindow)
@@ -257,11 +259,14 @@ void FNiagaraDataInterfaceProxySpectrum::UpdateSpectrum()
 		NumWindows++;
 
 		// Run deinterleave view over window.
-		FDeinterleaveView DeinterleaveView(InterleavedWindow, WindowBuffer, NumChannels);
+		FDeinterleaveView DeinterleaveView(InterleavedWindow, DeinterleavedBuffer, NumChannels);
 		for(FChannel Channel : DeinterleaveView)
 		{
+			Audio::ArrayMultiplyInPlace(WindowBuffer, Channel.Values);
+
 			// Perform FFT
 			FMemory::Memcpy(FFTInputBuffer.GetData(), Channel.Values.GetData(), sizeof(float) * Channel.Values.Num());
+
 			FFTAlgorithm->ForwardRealToComplex(FFTInputBuffer.GetData(), FFTOutputBuffer.GetData());
 
 			// Transflate FFTOutput to power spectrum
@@ -290,7 +295,7 @@ void FNiagaraDataInterfaceProxySpectrum::UpdateSpectrum()
 
 			Audio::ArrayMultiplyByConstantInPlace(Buffer, LinearScale);
 
-			Audio::ArrayPowerToDecibelInPlace(Buffer);
+			Audio::ArrayPowerToDecibelInPlace(Buffer, NoiseFloorDb);
 
 			Audio::ArraySubtractByConstantInPlace(Buffer, NoiseFloorDb);
 
@@ -315,7 +320,7 @@ void FNiagaraDataInterfaceProxySpectrum::SetAudioFormat(int32 InNumChannels, flo
 	{
 		int32 FFTSize = FFTAlgorithm.IsValid() ? FFTAlgorithm->Size() : 0;
 
-		ResizeSlidingBuffer(InNumChannels, FFTSize);
+		ResizeWindow(InNumChannels, FFTSize);
 	}
 
 	NumChannels = InNumChannels;
@@ -340,10 +345,19 @@ void FNiagaraDataInterfaceProxySpectrum::ResizeSpectrumBuffer(int32 InNumChannel
 	}
 }
 
-void FNiagaraDataInterfaceProxySpectrum::ResizeSlidingBuffer(int32 InNumChannels, int32 InFFTSize)
+void FNiagaraDataInterfaceProxySpectrum::ResizeWindow(int32 InNumChannels, int32 InFFTSize)
 {
-	int32 NumWindowSamples = FMath::Max(1, InFFTSize);
-	NumWindowSamples = FMath::Min(NumWindowSamples, 1024);
+	int32 NumWindowFrames = FMath::Max(1, InFFTSize);
+
+	NumWindowFrames = FMath::Min(NumWindowFrames, 1024);
+
+	WindowBuffer.Reset();
+	WindowBuffer.AddUninitialized(NumWindowFrames);
+
+	Audio::GenerateBlackmanWindow(WindowBuffer.GetData(), NumWindowFrames, 1, true);
+
+	int32 NumWindowSamples = NumWindowFrames;
+
 	if (InNumChannels > 0)
 	{
 		NumWindowSamples *= InNumChannels;
@@ -353,6 +367,7 @@ void FNiagaraDataInterfaceProxySpectrum::ResizeSlidingBuffer(int32 InNumChannels
 	int32 NumHopSamples = FMath::Max(1, NumWindowSamples / 2);
 
 	SlidingBuffer = MakeUnique<FSlidingBuffer>(NumWindowSamples, NumHopSamples);
+
 }
 
 void FNiagaraDataInterfaceProxySpectrum::ResizeAudioTransform(float InMinimumFrequency, float InMaximumFrequency, float InSampleRate, int32 InNumBands)
@@ -532,8 +547,8 @@ Audio::FFFTSettings FNiagaraDataInterfaceProxySpectrum::GetFFTSettings(float InM
 
 UNiagaraDataInterfaceAudioSpectrum::UNiagaraDataInterfaceAudioSpectrum(FObjectInitializer const& ObjectInitializer)
 	: Super(ObjectInitializer)
-	, Resolution(128)
-	, MinimumFrequency(40.f)
+	, Resolution(512)
+	, MinimumFrequency(55.f)
 	, MaximumFrequency(10000.f)
 	, NoiseFloorDb(-60.f)
 	
@@ -615,16 +630,21 @@ bool UNiagaraDataInterfaceAudioSpectrum::GetFunctionHLSL(const FNiagaraDataInter
 			void {FunctionName}(float In_NormalizedPosition, int In_ChannelIndex, out float Out_Val)
 			{
 				float FrameIndex = In_NormalizedPosition * {SpectrumResolution};
+				int MaxIndex = {SpectrumResolution} - 1;
 				int LowerIndex = floor(FrameIndex);
-				LowerIndex = LowerIndex < {SpectrumResolution} ? LowerIndex : {SpectrumResolution} - 1.0;
-				int UpperIndex =  LowerIndex < {SpectrumResolution} ? LowerIndex + 1.0 : LowerIndex;
+				LowerIndex = LowerIndex < {SpectrumResolution} ? LowerIndex : MaxIndex;
+				LowerIndex = LowerIndex >= 0 ? LowerIndex : 0;
+				int UpperIndex =  LowerIndex < MaxIndex ? LowerIndex + 1 : LowerIndex;
 				float Fraction = FrameIndex - LowerIndex;
+				Fraction = Fraction > 1.0 ? 1.0 : Fraction;
+				Fraction = Fraction < 0.0 ? 0.0 : Fraction;
 				float LowerValue = {SpectrumBuffer}.Load(In_ChannelIndex * {SpectrumResolution} + LowerIndex);
 				float UpperValue = {SpectrumBuffer}.Load(In_ChannelIndex * {SpectrumResolution} + UpperIndex);
 
 				Out_Val = lerp(LowerValue, UpperValue, Fraction);
 			}
 		)");
+
 		TMap<FString, FStringFormatArg> ArgsBounds = {
 			{TEXT("FunctionName"), FStringFormatArg(FunctionInfo.InstanceName)},
 			{TEXT("ChannelCount"), FStringFormatArg(NumChannelsName + ParamInfo.DataInterfaceHLSLSymbol)},
diff --git a/Engine/Plugins/Runtime/AudioSynesthesia/Source/AudioSynesthesia/Private/ConstantQAnalyzer.cpp b/Engine/Plugins/Runtime/AudioSynesthesia/Source/AudioSynesthesia/Private/ConstantQAnalyzer.cpp
--- a/Engine/Plugins/Runtime/AudioSynesthesia/Source/AudioSynesthesia/Private/ConstantQAnalyzer.cpp
+++ b/Engine/Plugins/Runtime/AudioSynesthesia/Source/AudioSynesthesia/Private/ConstantQAnalyzer.cpp
@@ -60,17 +60,17 @@ namespace Audio
 			switch (Settings.SpectrumType)
 			{
 				case ESpectrumType::MagnitudeSpectrum:
-					ArrayMagnitudeToDecibelInPlace(OutCQT);
+					ArrayMagnitudeToDecibelInPlace(OutCQT, -90.f);
 					break;
 
 				case ESpectrumType::PowerSpectrum:
-					ArrayPowerToDecibelInPlace(OutCQT);
+					ArrayPowerToDecibelInPlace(OutCQT, -90.f);
 					break;
 
 				default:
 					check(false);
 					//checkf(false, TEXT("Unhandled ESpectrumType %s"), GETENUMSTRING(ESpectrumType, Settings.SpectrumType));
-					ArrayPowerToDecibelInPlace(OutCQT);
+					ArrayPowerToDecibelInPlace(OutCQT, -90.f);
 			}
 		}
 	}
diff --git a/Engine/Plugins/Runtime/AudioSynesthesia/Source/AudioSynesthesia/Private/OnsetAnalyzer.cpp b/Engine/Plugins/Runtime/AudioSynesthesia/Source/AudioSynesthesia/Private/OnsetAnalyzer.cpp
--- a/Engine/Plugins/Runtime/AudioSynesthesia/Source/AudioSynesthesia/Private/OnsetAnalyzer.cpp
+++ b/Engine/Plugins/Runtime/AudioSynesthesia/Source/AudioSynesthesia/Private/OnsetAnalyzer.cpp
@@ -131,7 +131,7 @@ namespace Audio
 		MelTransform->TransformArray(FFTSpectrumBuffer, MelSpectrum);
 
 		// Apply decibel scaling 
-		ArrayPowerToDecibelInPlace(MelSpectrum);
+		ArrayPowerToDecibelInPlace(MelSpectrum, -90.f);
 
 		// Clamp to noise floor
 		ArrayClampMinInPlace(MelSpectrum, Settings.NoiseFloorDb);
diff --git a/Engine/Source/Runtime/SignalProcessing/Private/ConstantQ.cpp b/Engine/Source/Runtime/SignalProcessing/Private/ConstantQ.cpp
--- a/Engine/Source/Runtime/SignalProcessing/Private/ConstantQ.cpp
+++ b/Engine/Source/Runtime/SignalProcessing/Private/ConstantQ.cpp
@@ -102,24 +102,39 @@ namespace Audio
 			// Sanity check the CQT bins to make sure we the bandwidth wasn't so small that the array is essentially empty
 			if (OffsetBandWeights.Num() > 0)
 			{
-				if (OffsetBandWeights.FilterByPredicate([](float InVal) { return InVal >= 0.5f; }).Num() < 1)
+
+				float GaussianSum = 0.f;
+				Audio::ArraySum(OffsetBandWeights, GaussianSum);
+
+				if (GaussianSum < 0.95f)
 				{
-					// Hit a condition where all values in band weights are below 0.5f. 
 					// It's a bit of an arbitrary threshold, but it tells us that the bandwidth
 					// is low enough and the FFT granularity course enough to where our pseudo 
 					// cqt windows will likely miss data. In this case we force the window to 
-					// have a single value of 1 at the nearest FFT bin.
+					// interpolate between two nearest fft bins.
 					FMemory::Memset(OffsetBandWeights.GetData(), 0, OffsetBandWeights.Num() * sizeof(float));
-					int32 NearestIndex = FMath::RoundToInt(InFFTSize * BandCenter / InSampleRate) - OffsetBandWeightsIndex;
+					float Position = InFFTSize * BandCenter / InSampleRate - static_cast<float>(OffsetBandWeightsIndex);
+					int32 LowerIndex = FMath::FloorToInt(Position);
+					int32 UpperIndex = LowerIndex + 1;
 
-					if (NearestIndex < OffsetBandWeights.Num())
+					if ((LowerIndex >= 0) && (UpperIndex < OffsetBandWeights.Num()))
 					{
-						NearestIndex = FMath::Clamp(NearestIndex, 0, OffsetBandWeights.Num() - 1);
-						OffsetBandWeights[NearestIndex] = 1.f / static_cast<float>(InFFTSize);
+						OffsetBandWeights[LowerIndex] = 1.f - (Position - static_cast<float>(LowerIndex));
+						OffsetBandWeights[UpperIndex] = 1.f - OffsetBandWeights[LowerIndex];
 
 						// Need to update digital bandwidth to be bandwidth of FFTbin
 						DigitalBandWidth = 1.f / static_cast<float>(InFFTSize);
 					}
+					else if ((LowerIndex >= 0) && (LowerIndex < OffsetBandWeights.Num()))
+					{
+						OffsetBandWeights[LowerIndex] = 1.f;
+						DigitalBandWidth = 1.f / static_cast<float>(InFFTSize);
+					}
+					else if ((UpperIndex >= 0) && (UpperIndex < OffsetBandWeights.Num()))
+					{
+						OffsetBandWeights[UpperIndex] = 1.f;
+						DigitalBandWidth = 1.f / static_cast<float>(InFFTSize);
+					}
 				}
 			}
 
diff --git a/Engine/Source/Runtime/SignalProcessing/Private/FloatArrayMath.cpp b/Engine/Source/Runtime/SignalProcessing/Private/FloatArrayMath.cpp
--- a/Engine/Source/Runtime/SignalProcessing/Private/FloatArrayMath.cpp
+++ b/Engine/Source/Runtime/SignalProcessing/Private/FloatArrayMath.cpp
@@ -11,8 +11,54 @@ namespace Audio
 		const float Loge10 = FMath::Loge(10.f);
 		const int32 SimdMask = 0xFFFFFFFC;
 		const int32 NotSimdMask = 0x00000003;
+	}
 
-		
+	void ArraySum(TArrayView<const float> InValues, float& OutSum)
+	{
+		OutSum = 0.f;
+
+		int32 Num = InValues.Num();
+		const float* InData = InValues.GetData();
+
+		for (int32 i = 0; i < Num; i++)
+		{
+			OutSum += InData[i];
+		}
+	}
+
+	void ArraySum(const AlignedFloatBuffer& InValues, float& OutSum)
+	{
+		OutSum = 0.f;
+
+		const int32 Num = InValues.Num();
+		const int32 NumToSimd = Num & MathIntrinsics::SimdMask;
+		const int32 NumNotToSimd = Num & MathIntrinsics::NotSimdMask;
+
+		const float* InData = InValues.GetData();
+
+		if (NumToSimd)
+		{
+			VectorRegister Total = VectorSetFloat1(0.f);
+
+			for (int32 i = 0; i < NumToSimd; i += 4)
+			{
+				VectorRegister VectorData = VectorLoadAligned(&InData[i]);
+				Total = VectorAdd(Total, VectorData);
+			}
+
+			OutSum = VectorGetComponent(Total, 0) + VectorGetComponent(Total, 1) + VectorGetComponent(Total, 2) + VectorGetComponent(Total, 3);
+		}
+
+		if (NumNotToSimd)
+		{
+			TArrayView<const float> ValuesView(&InData[NumToSimd],  NumNotToSimd);
+
+			float ExtraSum = 0.f;
+
+			ArraySum(ValuesView, ExtraSum);
+
+			OutSum += ExtraSum;
+		}
 	}
 
 	void ArrayCumulativeSum(TArrayView<const float> InView, TArray<float>& OutData)
@@ -304,6 +350,44 @@ namespace Audio
 		}
 	}
 
+	void ArrayMultiplyInPlace(TArrayView<const float> InValues1, TArrayView<float> InValues2)
+	{
+		const int32 Num = InValues1.Num();
+
+		const float* InData1 = InValues1.GetData();
+		float* InData2 = InValues2.GetData();
+
+		for (int32 i = 0; i < Num; i++)
+		{
+			InData2[i] *= InData1[i];
+		}
+	}
+
+	void ArrayMultiplyInPlace(const AlignedFloatBuffer& InValues1, AlignedFloatBuffer& InValues2)
+	{
+		check(InValues1.Num() == InValues2.Num());
+
+		const int32 Num = InValues1.Num();
+		const int32 NumToSimd = Num & MathIntrinsics::SimdMask;
+		const int32 NumNotToSimd = Num & MathIntrinsics::NotSimdMask;
+
+		const float* InData1 = InValues1.GetData();
+		float* InData2 = InValues2.GetData();
+
+		if (NumToSimd)
+		{
+			MultiplyBuffersInPlace(InData1, InData2, NumToSimd);
+		}
+
+		if (NumNotToSimd)
+		{
+			TArrayView<const float> ValuesView1(&InData1[NumToSimd],  NumNotToSimd);
+			TArrayView<float> ValuesView2(&InData2[NumToSimd],  NumNotToSimd);
+
+			ArrayMultiplyInPlace(ValuesView1, ValuesView2);
+		}
+	}
+
 	void ArrayMultiplyByConstantInPlace(TArrayView<float> InValues, float InMultiplier)
 	{
 		const int32 Num = InValues.Num();
@@ -334,7 +418,6 @@ namespace Audio
 
 			ArrayMultiplyByConstantInPlace(ValuesView, InMultiplier);
 		}
-		
 	}
 
 	void ArrayAddInPlace(TArrayView<const float> InValues, TArrayView<float> InAccumulateValues)
@@ -440,17 +523,21 @@ namespace Audio
 		}
 	}
 
-	void ArrayMagnitudeToDecibelInPlace(TArrayView<float> InValues)
+	void ArrayMagnitudeToDecibelInPlace(TArrayView<float> InValues, float InMinimumDb)
 	{
 		const int32 Num = InValues.Num();
 		float* InValuesData = InValues.GetData();
+
+		const float Minimum = FMath::Exp(InMinimumDb * MathIntrinsics::Loge10 / 20.f);
+
 		for (int32 i = 0; i < Num; i++)
 		{
+			InValuesData[i] = FMath::Max(InValuesData[i], Minimum);
 			InValuesData[i] = 20.f * FMath::Loge(InValuesData[i]) / MathIntrinsics::Loge10;
 		}
 	}
 
-	void ArrayMagnitudeToDecibelInPlace(AlignedFloatBuffer& InValues)
+	void ArrayMagnitudeToDecibelInPlace(AlignedFloatBuffer& InValues, float InMinimumDb)
 	{
 		const int32 Num = InValues.Num();
 		const int32 NumToSimd = Num & MathIntrinsics::SimdMask;
@@ -458,12 +545,17 @@ namespace Audio
 
 		float* InData = InValues.GetData();
 
-		const VectorRegister VectorScale = VectorSetFloat1(20.f / MathIntrinsics::Loge10);
+		const float Scale = 20.f / MathIntrinsics::Loge10;
+		const float Minimum = FMath::Exp(InMinimumDb * MathIntrinsics::Loge10 / 20.f);
+
+		const VectorRegister VectorScale = VectorSetFloat1(Scale);
+		const VectorRegister VectorMinimum = VectorSetFloat1(Minimum);
 
 		for (int32 i = 0; i < NumToSimd; i += 4)
 		{
 			VectorRegister VectorData = VectorLoadAligned(&InData[i]);
-
+			
+			VectorData = VectorMax(VectorData, VectorMinimum);
 			VectorData = VectorLog(VectorData);
 			VectorData = VectorMultiply(VectorData, VectorScale);
 
@@ -473,21 +565,25 @@ namespace Audio
 		if (NumNotToSimd)
 		{
 			TArrayView<float> InView(&InData[NumToSimd], NumNotToSimd);
-			ArrayMagnitudeToDecibelInPlace(InView);
+			ArrayMagnitudeToDecibelInPlace(InView, InMinimumDb);
 		}
 	}
 
-	void ArrayPowerToDecibelInPlace(TArrayView<float> InValues)
+	void ArrayPowerToDecibelInPlace(TArrayView<float> InValues, float InMinimumDb)
 	{
 		const int32 Num = InValues.Num();
 		float* InValuesData = InValues.GetData();
+
+		const float Minimum = FMath::Exp(InMinimumDb * MathIntrinsics::Loge10 / 10.f);
+
 		for (int32 i = 0; i < Num; i++)
 		{
+			InValuesData[i] = FMath::Max(InValuesData[i], Minimum);
 			InValuesData[i] = 10.f * FMath::Loge(InValuesData[i]) / MathIntrinsics::Loge10;
 		}
 	}
 
-	void ArrayPowerToDecibelInPlace(AlignedFloatBuffer& InValues)
+	void ArrayPowerToDecibelInPlace(AlignedFloatBuffer& InValues, float InMinimumDb)
 	{
 		const int32 Num = InValues.Num();
 		const int32 NumToSimd = Num & MathIntrinsics::SimdMask;
@@ -495,12 +591,17 @@ namespace Audio
 
 		float* InData = InValues.GetData();
 
-		const VectorRegister VectorScale = VectorSetFloat1(10.f / MathIntrinsics::Loge10);
+		const float Scale = 10.f / MathIntrinsics::Loge10;
+		const float Minimum = FMath::Exp(InMinimumDb * MathIntrinsics::Loge10 / 10.f);
+
+		const VectorRegister VectorMinimum = VectorSetFloat1(Minimum);
+		const VectorRegister VectorScale = VectorSetFloat1(Scale);
 
 		for (int32 i = 0; i < NumToSimd; i += 4)
 		{
 			VectorRegister VectorData = VectorLoadAligned(&InData[i]);
 
+			VectorData = VectorMax(VectorData, VectorMinimum);
 			VectorData = VectorLog(VectorData);
 			VectorData = VectorMultiply(VectorData, VectorScale);
 
@@ -510,7 +611,7 @@ namespace Audio
 		if (NumNotToSimd)
 		{
 			TArrayView<float> InView(&InData[NumToSimd], NumNotToSimd);
-			ArrayPowerToDecibelInPlace(InView);
+			ArrayPowerToDecibelInPlace(InView, InMinimumDb);
 		}
 	}
 
diff --git a/Engine/Source/Runtime/SignalProcessing/Public/DSP/DeinterleaveView.h b/Engine/Source/Runtime/SignalProcessing/Public/DSP/DeinterleaveView.h
--- a/Engine/Source/Runtime/SignalProcessing/Public/DSP/DeinterleaveView.h
+++ b/Engine/Source/Runtime/SignalProcessing/Public/DSP/DeinterleaveView.h
@@ -25,10 +25,10 @@ namespace Audio
 		class TChannel
 		{
 			public:
-				const TArray<T, InAllocator>& Values;
+				TArray<T, InAllocator>& Values;
 				const int32 ChannelIndex;
 
-				TChannel(const TArray<T, InAllocator>& InValues, const int32 InChannelIndex)
+				TChannel(TArray<T, InAllocator>& InValues, const int32 InChannelIndex)
 				:	Values(InValues)
 				,	ChannelIndex(InChannelIndex)
 				{}
@@ -71,12 +71,6 @@ namespace Audio
 				}
 
 
-				// prepare the array
-				ArrayToFill.Reset(DeinterleaveView.NumElementsPerChannel);
-				if (DeinterleaveView.NumElementsPerChannel > 0)
-				{
-					ArrayToFill.AddUninitialized(DeinterleaveView.NumElementsPerChannel);
-				}
 			}
 
 			/** Increment the iterator forward by one channel */
@@ -114,6 +108,13 @@ namespace Audio
 				}
 				else
 				{
+					// prepare the array
+					ArrayToFill.Reset(DeinterleaveView.NumElementsPerChannel);
+					if (DeinterleaveView.NumElementsPerChannel > 0)
+					{
+						ArrayToFill.AddUninitialized(DeinterleaveView.NumElementsPerChannel);
+					}
+
 					// Fill array with deinterleave data
 					T* ArrayToFillData = ArrayToFill.GetData();
 					const T* InterleavedData = DeinterleaveView.InterleavedArray.GetData();
diff --git a/Engine/Source/Runtime/SignalProcessing/Public/DSP/FloatArrayMath.h b/Engine/Source/Runtime/SignalProcessing/Public/DSP/FloatArrayMath.h
--- a/Engine/Source/Runtime/SignalProcessing/Public/DSP/FloatArrayMath.h
+++ b/Engine/Source/Runtime/SignalProcessing/Public/DSP/FloatArrayMath.h
@@ -7,6 +7,9 @@
 
 namespace Audio
 {
+	SIGNALPROCESSING_API void ArraySum(TArrayView<const float> InValues, float& OutSum);
+	SIGNALPROCESSING_API void ArraySum(const AlignedFloatBuffer& InValues, float& OutSum);
+
 	/** Cumulative sum of array.
 	 *
 	 *  InView contains data to be cumulatively summed.
@@ -82,6 +85,10 @@ namespace Audio
 	 */ 
 	SIGNALPROCESSING_API void ArrayMinMaxNormalize(TArrayView<const float> InView, TArray<float>& OutArray);
 
+	/** Multiplies two arrays element-wise. Stores result in InValues2*/
+	SIGNALPROCESSING_API void ArrayMultiplyInPlace(TArrayView<const float> InValues1, TArrayView<float> InValues2);
+	SIGNALPROCESSING_API void ArrayMultiplyInPlace(const AlignedFloatBuffer& InValues1, AlignedFloatBuffer& InValues2);
+
 	/** Multiplies each element in InView by InMultiplier */
 	SIGNALPROCESSING_API void ArrayMultiplyByConstantInPlace(TArrayView<float> InValues, float InMultiplier);
 	SIGNALPROCESSING_API void ArrayMultiplyByConstantInPlace(AlignedFloatBuffer& InValues, float InMultiplier);
@@ -107,12 +114,12 @@ namespace Audio
 	SIGNALPROCESSING_API void ArraySubtractByConstantInPlace(AlignedFloatBuffer& InValues, float InSubtrahend);
 
 	/** Convert magnitude values to decibel values in place. db = 20 * log10(val) */
-	SIGNALPROCESSING_API void ArrayMagnitudeToDecibelInPlace(TArrayView<float> InValues);
-	SIGNALPROCESSING_API void ArrayMagnitudeToDecibelInPlace(AlignedFloatBuffer& InValues);
+	SIGNALPROCESSING_API void ArrayMagnitudeToDecibelInPlace(TArrayView<float> InValues, float InMinimumDb);
+	SIGNALPROCESSING_API void ArrayMagnitudeToDecibelInPlace(AlignedFloatBuffer& InValues, float InMinimumDb);
 
 	/** Convert power values to decibel values in place. db = 10 * log10(val) */
-	SIGNALPROCESSING_API void ArrayPowerToDecibelInPlace(TArrayView<float> InValues);
-	SIGNALPROCESSING_API void ArrayPowerToDecibelInPlace(AlignedFloatBuffer& InValues);
+	SIGNALPROCESSING_API void ArrayPowerToDecibelInPlace(TArrayView<float> InValues, float InMinimumDb);
+	SIGNALPROCESSING_API void ArrayPowerToDecibelInPlace(AlignedFloatBuffer& InValues, float InMinimumDb);
 
 
 	/** Compute power of complex data. Out[i] = Complex[2 * i] * Complex[2 * i] + Complex[2 * i + 1] * Complex[2 * i + 1] */
