Fix for incorrect playback state on a sound (regression from short fade-in time fix)

diff --git a/Engine/Source/Runtime/SignalProcessing/Private/VolumeFader.cpp b/Engine/Source/Runtime/SignalProcessing/Private/VolumeFader.cpp
--- a/Engine/Source/Runtime/SignalProcessing/Private/VolumeFader.cpp
+++ b/Engine/Source/Runtime/SignalProcessing/Private/VolumeFader.cpp
@@ -235,28 +235,36 @@ namespace Audio
 		}
 		Target = Alpha;
 		FadeCurve = EFaderCurve::Linear;
+		Elapsed = ActiveDuration;
 		FadeDuration = -1.0f;
 	}
 
 	void FVolumeFader::Update(float InDeltaTime)
 	{
-		Elapsed += InDeltaTime;
-
 		if (!IsFading() || !IsActive())
 		{
 			return;
 		}
 
+		Elapsed += InDeltaTime;
+
 		// Keep stepping towards target & clamp until fade duration has expired
 		// Choose min/max bound and clamp dt to prevent unwanted spikes in volume
 		float MinValue = 0.0f;
 		float MaxValue = 0.0f;
-		if (Alpha < Target)
+
+		if (FadeDuration < Elapsed)
+		{
+			Alpha = Target;
+			StopFade();
+			return;
+		}
+		else if (Alpha < Target)
 		{
 			MinValue = Alpha;
 			MaxValue = Target;
 		}
-		else
+		else 
 		{
 			MinValue = Target;
 			MaxValue = Alpha;
@@ -265,6 +273,7 @@ namespace Audio
 		Alpha += (Target - Alpha) * InDeltaTime / (FadeDuration - Elapsed);
 		Alpha = FMath::Clamp(Alpha, MinValue, MaxValue);
 
+
 		// Optimization that avoids fade calculations once alpha reaches target
 		if (InDeltaTime > SMALL_NUMBER && FMath::IsNearlyEqual(Alpha, Target))
 		{
