Fixing crash in CQT generation

diff --git a/Engine/Source/Runtime/SignalProcessing/Private/FloatArrayMath.cpp b/Engine/Source/Runtime/SignalProcessing/Private/FloatArrayMath.cpp
--- a/Engine/Source/Runtime/SignalProcessing/Private/FloatArrayMath.cpp
+++ b/Engine/Source/Runtime/SignalProcessing/Private/FloatArrayMath.cpp
@@ -683,10 +683,14 @@ namespace Audio
 		FRow EmptyRow;
 		EmptyRow.StartIndex = 0;
 
-		// Fill up the kernel with emptp rows
+		// Fill up the kernel with empty rows
 		Kernel.Init(EmptyRow, NumOut);
 	}
 
+	FContiguousSparse2DKernelTransform::~FContiguousSparse2DKernelTransform()
+	{
+	}
+
 	int32 FContiguousSparse2DKernelTransform::GetNumInElements() const
 	{
 		return NumIn;
diff --git a/Engine/Source/Runtime/SignalProcessing/Public/DSP/FloatArrayMath.h b/Engine/Source/Runtime/SignalProcessing/Public/DSP/FloatArrayMath.h
--- a/Engine/Source/Runtime/SignalProcessing/Public/DSP/FloatArrayMath.h
+++ b/Engine/Source/Runtime/SignalProcessing/Public/DSP/FloatArrayMath.h
@@ -145,17 +145,21 @@ namespace Audio
 	{
 		struct FRow
 		{
-			int32 StartIndex;
+			int32 StartIndex = 0;
 			TArray<float> OffsetValues;
 		};
 
 	public:
+		FContiguousSparse2DKernelTransform(const FContiguousSparse2DKernelTransform& ) = delete;
+		FContiguousSparse2DKernelTransform(const FContiguousSparse2DKernelTransform&& ) = delete;
+		FContiguousSparse2DKernelTransform& operator=(const FContiguousSparse2DKernelTransform& ) = delete;
 
 		/**
 		 * NumInElements sets the expected number of input array elements as well as the number of elements in a row.
 		 * NumOutElements sets the number of output array elements as well as the number or rows.
 		 */
 		FContiguousSparse2DKernelTransform(const int32 NumInElements, const int32 NumOutElements);
+		virtual ~FContiguousSparse2DKernelTransform();
 
 		/** Returns the required size of the input array */
 		int32 GetNumInElements() const;
