Audio: edigrate volume fader fix to 4.25 (Bug in StopDelayed)

diff --git a/Engine/Source/Runtime/SignalProcessing/Private/VolumeFader.cpp b/Engine/Source/Runtime/SignalProcessing/Private/VolumeFader.cpp
--- a/Engine/Source/Runtime/SignalProcessing/Private/VolumeFader.cpp
+++ b/Engine/Source/Runtime/SignalProcessing/Private/VolumeFader.cpp
@@ -184,6 +184,7 @@ namespace Audio
 		else
 		{
 			ActiveDuration = InDuration;
+			Elapsed = 0.0f;
 		}
 	}
 
