Adding spectral analysis delegate to sound submix bp library. Also adding constant q to spectral analyzer.
Utilizing optimized FFT in spectral analyzer

diff --git a/Engine/Source/Runtime/AudioMixer/Private/AudioMixerBlueprintLibrary.cpp b/Engine/Source/Runtime/AudioMixer/Private/AudioMixerBlueprintLibrary.cpp
--- a/Engine/Source/Runtime/AudioMixer/Private/AudioMixerBlueprintLibrary.cpp
+++ b/Engine/Source/Runtime/AudioMixer/Private/AudioMixerBlueprintLibrary.cpp
@@ -9,6 +9,7 @@
 #include "ContentStreaming.h"
 #include "AudioCompressionSettingsUtils.h"
 #include "Async/Async.h"
+#include "Sound/SoundSubmix.h"
 #include "Sound/SoundEffectPreset.h"
 
 // This is our global recording task:
@@ -237,7 +238,7 @@ USoundWave* UAudioMixerBlueprintLibrary::StopRecordingOutput(const UObject* Worl
 	return nullptr;
 }
 
-void UAudioMixerBlueprintLibrary::PauseRecordingOutput(const UObject* WorldContextObject, USoundSubmix* SubmixToPause /* = nullptr */)
+void UAudioMixerBlueprintLibrary::PauseRecordingOutput(const UObject* WorldContextObject, USoundSubmix* SubmixToPause)
 {
 	if (Audio::FMixerDevice* MixerDevice = GetAudioMixerDeviceFromWorldContext(WorldContextObject))
 	{
@@ -249,7 +250,7 @@ void UAudioMixerBlueprintLibrary::PauseRecordingOutput(const UObject* WorldConte
 	}
 }
 
-void UAudioMixerBlueprintLibrary::ResumeRecordingOutput(const UObject* WorldContextObject, USoundSubmix* SubmixToResume /* = nullptr */)
+void UAudioMixerBlueprintLibrary::ResumeRecordingOutput(const UObject* WorldContextObject, USoundSubmix* SubmixToResume)
 {
 	if (Audio::FMixerDevice* MixerDevice = GetAudioMixerDeviceFromWorldContext(WorldContextObject))
 	{
@@ -261,12 +262,11 @@ void UAudioMixerBlueprintLibrary::ResumeRecordingOutput(const UObject* WorldCont
 	}
 }
 
-void UAudioMixerBlueprintLibrary::StartAnalyzingOutput(const UObject* WorldContextObject, USoundSubmix* SubmixToAnalyze /*= nullptr*/, EFFTSize FFTSize /*= EFFTSize::Default*/, EFFTPeakInterpolationMethod InterpolationMethod /*= Linear*/, EFFTWindowType WindowType /*= EFFTWindowType::Hamming*/, float HopSize /*= 0*/)
+void UAudioMixerBlueprintLibrary::StartAnalyzingOutput(const UObject* WorldContextObject, USoundSubmix* SubmixToAnalyze, EFFTSize FFTSize, EFFTPeakInterpolationMethod InterpolationMethod, EFFTWindowType WindowType, float HopSize, EAudioSpectrumType AudioSpectrumType)
 {
 	if (Audio::FMixerDevice* MixerDevice = GetAudioMixerDeviceFromWorldContext(WorldContextObject))
 	{
-		Audio::FSpectrumAnalyzerSettings Settings = Audio::FSpectrumAnalyzerSettings();
-		PopulateSpectrumAnalyzerSettings(FFTSize, InterpolationMethod, WindowType, HopSize, Settings);
+		Audio::FSpectrumAnalyzerSettings Settings = USoundSubmix::GetSpectrumAnalyzerSettings(FFTSize, InterpolationMethod, WindowType, HopSize, AudioSpectrumType);
 		MixerDevice->StartSpectrumAnalysis(SubmixToAnalyze, Settings);
 	}
 	else
@@ -275,7 +275,7 @@ void UAudioMixerBlueprintLibrary::StartAnalyzingOutput(const UObject* WorldConte
 	}
 }
 
-void UAudioMixerBlueprintLibrary::StopAnalyzingOutput(const UObject* WorldContextObject, USoundSubmix* SubmixToStopAnalyzing /*= nullptr*/)
+void UAudioMixerBlueprintLibrary::StopAnalyzingOutput(const UObject* WorldContextObject, USoundSubmix* SubmixToStopAnalyzing)
 {
 	if (Audio::FMixerDevice* MixerDevice = GetAudioMixerDeviceFromWorldContext(WorldContextObject))
 	{
@@ -503,64 +503,3 @@ bool UAudioMixerBlueprintLibrary::IsAudioBusActive(const UObject* WorldContextOb
 	}
 }
 
-void UAudioMixerBlueprintLibrary::PopulateSpectrumAnalyzerSettings(EFFTSize FFTSize, EFFTPeakInterpolationMethod InterpolationMethod, EFFTWindowType WindowType, float HopSize, Audio::FSpectrumAnalyzerSettings &OutSettings)
-{
-	switch (FFTSize)
-	{
-	case EFFTSize::DefaultSize:
-		OutSettings.FFTSize = Audio::FSpectrumAnalyzerSettings::EFFTSize::Default;
-		break;
-	case EFFTSize::Min:
-		OutSettings.FFTSize = Audio::FSpectrumAnalyzerSettings::EFFTSize::Min_64;
-		break;
-	case EFFTSize::Small:
-		OutSettings.FFTSize = Audio::FSpectrumAnalyzerSettings::EFFTSize::Small_256;
-		break;
-	case EFFTSize::Medium:
-		OutSettings.FFTSize = Audio::FSpectrumAnalyzerSettings::EFFTSize::Medium_512;
-		break;
-	case EFFTSize::Large:
-		OutSettings.FFTSize = Audio::FSpectrumAnalyzerSettings::EFFTSize::Large_1024;
-		break;
-	case EFFTSize::Max:
-		OutSettings.FFTSize = Audio::FSpectrumAnalyzerSettings::EFFTSize::TestLarge_4096;
-		break;
-	default:
-		break;
-	}
-
-	switch (InterpolationMethod)
-	{
-	case EFFTPeakInterpolationMethod::NearestNeighbor:
-		OutSettings.InterpolationMethod = Audio::FSpectrumAnalyzerSettings::EPeakInterpolationMethod::NearestNeighbor;
-		break;
-	case EFFTPeakInterpolationMethod::Linear:
-		OutSettings.InterpolationMethod = Audio::FSpectrumAnalyzerSettings::EPeakInterpolationMethod::Linear;
-		break;
-	case EFFTPeakInterpolationMethod::Quadratic:
-		OutSettings.InterpolationMethod = Audio::FSpectrumAnalyzerSettings::EPeakInterpolationMethod::Quadratic;
-		break;
-	default:
-		break;
-	}
-
-	switch (WindowType)
-	{
-	case EFFTWindowType::None:
-		OutSettings.WindowType = Audio::EWindowType::None;
-		break;
-	case EFFTWindowType::Hamming:
-		OutSettings.WindowType = Audio::EWindowType::Hamming;
-		break;
-	case EFFTWindowType::Hann:
-		OutSettings.WindowType = Audio::EWindowType::Hann;
-		break;
-	case EFFTWindowType::Blackman:
-		OutSettings.WindowType = Audio::EWindowType::Blackman;
-		break;
-	default:
-		break;
-	}
-
-	OutSettings.HopSize = HopSize;
-}
diff --git a/Engine/Source/Runtime/AudioMixer/Private/AudioMixerDevice.cpp b/Engine/Source/Runtime/AudioMixer/Private/AudioMixerDevice.cpp
--- a/Engine/Source/Runtime/AudioMixer/Private/AudioMixerDevice.cpp
+++ b/Engine/Source/Runtime/AudioMixer/Private/AudioMixerDevice.cpp
@@ -449,7 +449,7 @@ namespace Audio
 
 		// Loop through any envelope-following submixes and perform any broadcasting of envelope data if needed
 		TArray<float> SubmixEnvelopeData;
-		for (USoundSubmix* SoundSubmix : EnvelopeFollowingSubmixes)
+		for (USoundSubmix* SoundSubmix : DelegateBoundSubmixes)
 		{
 			if (SoundSubmix)
 			{
@@ -463,7 +463,7 @@ namespace Audio
 					Audio::FMixerSubmixPtr ThisSubmixPtr = SubmixPtr.Pin();
 					if (ThisSubmixPtr.IsValid())
 					{
-						ThisSubmixPtr->BroadcastEnvelope();
+						ThisSubmixPtr->BroadcastDelegates();
 					}
 				});
 			}
@@ -1752,7 +1752,7 @@ namespace Audio
 			MasterSubmixPtr->StartEnvelopeFollowing(InSubmix->EnvelopeFollowerAttackTime, InSubmix->EnvelopeFollowerReleaseTime);
 		}
 
-		EnvelopeFollowingSubmixes.AddUnique(InSubmix);
+		DelegateBoundSubmixes.AddUnique(InSubmix);
 	}
 
 	void FMixerDevice::StopEnvelopeFollowing(USoundSubmix* InSubmix)
@@ -1784,7 +1784,7 @@ namespace Audio
 			MasterSubmixPtr->StopEnvelopeFollowing();
 		}
 
-		EnvelopeFollowingSubmixes.RemoveSingleSwap(InSubmix);
+		DelegateBoundSubmixes.RemoveSingleSwap(InSubmix);
 	}
 
 	void FMixerDevice::AddEnvelopeFollowerDelegate(USoundSubmix* InSubmix, const FOnSubmixEnvelopeBP& OnSubmixEnvelopeBP)
@@ -1845,6 +1845,8 @@ namespace Audio
 
 			MasterSubmixPtr->StartSpectrumAnalysis(InSettings);
 		}
+
+		DelegateBoundSubmixes.AddUnique(InSubmix);
 	}
 
 	void FMixerDevice::StopSpectrumAnalysis(USoundSubmix* InSubmix)
@@ -1874,6 +1876,9 @@ namespace Audio
 
 			MasterSubmixPtr->StopSpectrumAnalysis();
 		}
+
+		DelegateBoundSubmixes.RemoveSingleSwap(InSubmix);
+
 	}
 
 	void FMixerDevice::GetMagnitudesForFrequencies(USoundSubmix* InSubmix, const TArray<float>& InFrequencies, TArray<float>& OutMagnitudes)
@@ -1910,6 +1915,36 @@ namespace Audio
 		}
 	}
 
+	void FMixerDevice::AddSpectralAnalysisDelegate(USoundSubmix* InSubmix, const TArray<FSoundSubmixSpectralAnalysisBandSettings>& InBandSettings, const FOnSubmixSpectralAnalysisBP& OnSubmixSpectralAnalysisBP, float UpdateRate)
+	{
+		if (!IsInAudioThread())
+		{
+			DECLARE_CYCLE_STAT(TEXT("FAudioThreadTask.AddFFTDelegate"), STAT_AddFFTDelegate, STATGROUP_AudioThreadCommands);
+
+			FAudioThread::RunCommandOnAudioThread([this, InSubmix, InBandSettings, OnSubmixSpectralAnalysisBP, UpdateRate]()
+			{
+				CSV_SCOPED_TIMING_STAT(Audio, AddEnvelopeFollowerDelegate);
+				AddSpectralAnalysisDelegate(InSubmix, InBandSettings, OnSubmixSpectralAnalysisBP, UpdateRate);
+			}, GET_STATID(STAT_AddFFTDelegate));
+			return;
+		}
+
+		// if we can find the submix here, record that submix. Otherwise, just record the master submix.
+		FMixerSubmixPtr FoundSubmix = GetSubmixInstance(InSubmix).Pin();
+		if (FoundSubmix.IsValid())
+		{
+			FoundSubmix->AddSpectralAnalysisDelegate(OnSubmixSpectralAnalysisBP, InBandSettings, UpdateRate);
+		}
+		else
+		{
+			FMixerSubmixWeakPtr MasterSubmix = GetMasterSubmix();
+			FMixerSubmixPtr MasterSubmixPtr = MasterSubmix.Pin();
+			check(MasterSubmixPtr.IsValid());
+
+			MasterSubmixPtr->AddSpectralAnalysisDelegate(OnSubmixSpectralAnalysisBP, InBandSettings, UpdateRate);
+		}
+	}
+
 	void FMixerDevice::RegisterSubmixBufferListener(ISubmixBufferListener* InSubmixBufferListener, USoundSubmix* InSubmix)
 	{
 		if (!IsInAudioThread())
diff --git a/Engine/Source/Runtime/AudioMixer/Private/AudioMixerSubmix.cpp b/Engine/Source/Runtime/AudioMixer/Private/AudioMixerSubmix.cpp
--- a/Engine/Source/Runtime/AudioMixer/Private/AudioMixerSubmix.cpp
+++ b/Engine/Source/Runtime/AudioMixer/Private/AudioMixerSubmix.cpp
@@ -1485,13 +1485,88 @@ namespace Audio
 		OnSubmixEnvelope.AddUnique(OnSubmixEnvelopeBP);
 	}
 
+	void FMixerSubmix::AddSpectralAnalysisDelegate(const FOnSubmixSpectralAnalysisBP& OnSubmixSpectralAnalysisBP, const TArray<FSoundSubmixSpectralAnalysisBandSettings>& InBandSettings, float UpdateRate)
+	{
+		FSpectrumAnalysisDelegateInfo& NewDelegateInfo = SpectralAnalysisDelegates.AddDefaulted_GetRef();
+		
+		NewDelegateInfo.LastUpdateTime = -1.0f;
+		UpdateRate = FMath::Clamp(UpdateRate, 1.0f, 30.0f);
+		NewDelegateInfo.UpdateDelta = 1.0f / UpdateRate;
+
+		for (const FSoundSubmixSpectralAnalysisBandSettings& BandSettings : InBandSettings)
+		{
+			FSpectralAnalysisBandInfo NewBand;
+			NewBand.BandFrequency = BandSettings.BandFrequency;
+			NewBand.EnvelopeFollower.Init(UpdateRate, BandSettings.AttackTimeMsec, BandSettings.ReleaseTimeMsec);
+		
+			NewDelegateInfo.SpectralBands.Add(NewBand);
+		}
+
+		NewDelegateInfo.OnSubmixSpectralAnalysis.AddUnique(OnSubmixSpectralAnalysisBP);
+
+	}
+
 	void FMixerSubmix::StartSpectrumAnalysis(const FSpectrumAnalyzerSettings& InSettings)
 	{
+		bIsSpectrumAnalyzing = true;
+
+		// Construct a band extractor for each delegate listening to FFT results
+		for (FSpectrumAnalysisDelegateInfo& DelegateInfo : SpectralAnalysisDelegates)
+		{
+			// Build the band extractor here since we are given the FFT settings to use at this point
+			FSpectrumBandExtractorSettings BandExtractorSettings;
+
+			BandExtractorSettings.FFTSize = (int32)InSettings.FFTSize;
+			BandExtractorSettings.SampleRate = MixerDevice->GetSampleRate();
+
+			DelegateInfo.SpectrumBandExtractor = ISpectrumBandExtractor::CreateSpectrumBandExtractor(BandExtractorSettings);
+
+			// Determine the band extractor metric (i.e. spectrum type). This is the same for all bands in a given extractor.
+			ISpectrumBandExtractor::EMetric Metric = ISpectrumBandExtractor::EMetric::Magnitude;
+			switch (InSettings.SpectrumType)
+			{
+			case FSpectrumAnalyzerSettings::ESpectrumAnalyzerType::Power:
+				Metric = ISpectrumBandExtractor::EMetric::Power;
+				break;
+
+			case FSpectrumAnalyzerSettings::ESpectrumAnalyzerType::Decibel:
+				Metric = ISpectrumBandExtractor::EMetric::Decibel;
+				break;
+
+			default:
+				break;
+			}
+			
+			// TODO: @phil -- feels like maybe we can pass in the enum to AddBand vs 4 different Add Band functions?  Maybe not because of the Q-factor?
+
+			// Loop through each band and add them to the extractor
+			for (FSpectralAnalysisBandInfo& BandInfo : DelegateInfo.SpectralBands)
+			{
+				switch (InSettings.InterpolationMethod)
+				{
+					case FSpectrumAnalyzerSettings::EPeakInterpolationMethod::NearestNeighbor:
+						DelegateInfo.SpectrumBandExtractor->AddNearestNeighborBand(BandInfo.BandFrequency, Metric);
+						break;
+					case FSpectrumAnalyzerSettings::EPeakInterpolationMethod::Linear:
+						DelegateInfo.SpectrumBandExtractor->AddLerpBand(BandInfo.BandFrequency, Metric);
+						break;
+					case FSpectrumAnalyzerSettings::EPeakInterpolationMethod::Quadratic:
+						DelegateInfo.SpectrumBandExtractor->AddQuadraticBand(BandInfo.BandFrequency, Metric);
+						break;
+					case FSpectrumAnalyzerSettings::EPeakInterpolationMethod::ConstantQ:
+						// TODO: @phil -- how to compute q-factor?? I assume some math around bands?
+						DelegateInfo.SpectrumBandExtractor->AddConstantQBand(BandInfo.BandFrequency, 20.0f, Metric);
+						break;
+				}
+			}
+		}
+
 		SpectrumAnalyzer.Reset(new FSpectrumAnalyzer(InSettings, MixerDevice->GetSampleRate()));
 	}
 
 	void FMixerSubmix::StopSpectrumAnalysis()
 	{
+		bIsSpectrumAnalyzing = false;
 		SpectrumAnalyzer.Reset();
 	}
 
@@ -1550,18 +1625,18 @@ namespace Audio
 		TargetWetLevel = FMath::Clamp(InWetLevel, 0.0f, 1.0f);
 	}
 
-	void FMixerSubmix::BroadcastEnvelope()
+	void FMixerSubmix::BroadcastDelegates()
 	{
 		if (bIsEnvelopeFollowing)
 		{
 			// Get the envelope data
 			TArray<float> EnvelopeData;
+			EnvelopeData.AddUninitialized(EnvelopeNumChannels);
 
 			{
 				// Make the copy of the envelope values using a critical section
 				FScopeLock EnvelopeScopeLock(&EnvelopeCriticalSection);
 
-				EnvelopeData.AddUninitialized(EnvelopeNumChannels);
 				FMemory::Memcpy(EnvelopeData.GetData(), EnvelopeValues, sizeof(float)*EnvelopeNumChannels);
 			}
 
@@ -1572,6 +1647,56 @@ namespace Audio
 			}
 
 		}
+		
+		// If we're analyzing spectra and if we've got delegates setup
+		if (bIsSpectrumAnalyzing && SpectralAnalysisDelegates.Num() > 0)
+		{
+			if (ensureMsgf(SpectrumAnalyzer.IsValid(), TEXT("Analyzing spectrum with invalid spectrum analyzer")))
+			{
+				// New results array
+				TArray<float> SpectralResults;
+
+				//const TArray<float>& InFrequencies, TArray<float>& OutMagnitudes
+				for (FSpectrumAnalysisDelegateInfo& DelegateInfo : SpectralAnalysisDelegates)
+				{
+					const float CurrentTime = FPlatformTime::ToSeconds64(FPlatformTime::Cycles64());
+
+					// Don't update the spectral band until it's time since the last tick.
+					if (DelegateInfo.LastUpdateTime > 0.0f && ((CurrentTime - DelegateInfo.LastUpdateTime) < DelegateInfo.UpdateDelta))
+					{
+						continue;
+					}
+
+					DelegateInfo.LastUpdateTime = CurrentTime;
+
+					SpectralResults.Reset();
+
+					{
+						Audio::FSpectrumAnalyzerScopeLock AnalyzerLock(SpectrumAnalyzer.Get());
+
+						check(DelegateInfo.SpectrumBandExtractor.IsValid());
+						ISpectrumBandExtractor* Extractor = DelegateInfo.SpectrumBandExtractor.Get();
+
+						SpectrumAnalyzer->GetBands(*Extractor, SpectralResults);
+					}
+
+					// Feed the results through the band envelope followers
+					for (int32 ResultIndex = 0; ResultIndex < SpectralResults.Num(); ++ResultIndex)
+					{
+						if (ensure(ResultIndex < DelegateInfo.SpectralBands.Num()))
+						{
+							FSpectralAnalysisBandInfo& BandInfo = DelegateInfo.SpectralBands[ResultIndex];
+							SpectralResults[ResultIndex] = BandInfo.EnvelopeFollower.ProcessAudioNonClamped(SpectralResults[ResultIndex]);
+						}
+					}
+
+					if (DelegateInfo.OnSubmixSpectralAnalysis.IsBound())
+					{
+						DelegateInfo.OnSubmixSpectralAnalysis.Broadcast(SpectralResults);
+					}
+				}
+			}
+		}
 	}
 
 	FSoundfieldEncodingKey FMixerSubmix::GetKeyForSubmixEncoding()
diff --git a/Engine/Source/Runtime/AudioMixer/Public/AudioMixerBlueprintLibrary.h b/Engine/Source/Runtime/AudioMixer/Public/AudioMixerBlueprintLibrary.h
--- a/Engine/Source/Runtime/AudioMixer/Public/AudioMixerBlueprintLibrary.h
+++ b/Engine/Source/Runtime/AudioMixer/Public/AudioMixerBlueprintLibrary.h
@@ -15,62 +15,6 @@
 
 class USoundSubmix;
 
-UENUM(BlueprintType)
-enum class EFFTSize : uint8
-{
-	// 512
-	DefaultSize, 
-
-	// 64
-	Min,
-
-	// 256
-	Small,
-
-	// 512
-	Medium, 
-
-	// 1024
-	Large, 
-
-	// 4096
-	Max, 
-};
-
-UENUM()
-enum class EFFTPeakInterpolationMethod : uint8
-{
-	NearestNeighbor,
-	Linear,
-	Quadratic
-};
-
-UENUM()
-enum class EFFTWindowType : uint8
-{
-	// No window is applied. Technically a boxcar window.
-	None, 
-
-	// Mainlobe width of -3 dB and sidelove attenuation of ~-40 dB. Good for COLA.
-	Hamming,
-
-	// Mainlobe width of -3 dB and sidelobe attenuation of ~-30dB. Good for COLA.
-	Hann,
-
-	// Mainlobe width of -3 dB and sidelobe attenuation of ~-60db. Tricky for COLA.
-	Blackman
-};
-
-UENUM(BlueprintType)
-enum class EAudioSpectrumType : uint8
-{
-	// Spectrum frequency values are equal to magnitude of frequency.
-	MagnitudeSpectrum,
-
-	// Spectrum frequency values are equal to magnitude squared.
-	PowerSpectrum
-};
-
 /** 
 * Called when a load request for a sound has completed.
 */
@@ -134,7 +78,7 @@ public:
 
 	/** Start spectrum analysis of the audio output. By leaving the Submix To Analyze blank, you can analyze the master output of the game. */
 	UFUNCTION(BlueprintCallable, Category = "Audio|Analysis", meta = (WorldContext = "WorldContextObject", AdvancedDisplay = 1))
-	static void StartAnalyzingOutput(const UObject* WorldContextObject, USoundSubmix* SubmixToAnalyze = nullptr, EFFTSize FFTSize = EFFTSize::DefaultSize, EFFTPeakInterpolationMethod InterpolationMethod = EFFTPeakInterpolationMethod::Linear, EFFTWindowType WindowType = EFFTWindowType::Hann, float HopSize = 0);
+	static void StartAnalyzingOutput(const UObject* WorldContextObject, USoundSubmix* SubmixToAnalyze = nullptr, EFFTSize FFTSize = EFFTSize::DefaultSize, EFFTPeakInterpolationMethod InterpolationMethod = EFFTPeakInterpolationMethod::Linear, EFFTWindowType WindowType = EFFTWindowType::Hann, float HopSize = 0, EAudioSpectrumType SpectrumType = EAudioSpectrumType::MagnitudeSpectrum);
 
 	/** Start spectrum analysis of the audio output. By leaving the Submix To Stop Analyzing blank, you can analyze the master output of the game. */
 	UFUNCTION(BlueprintCallable, Category = "Audio|Analysis", meta = (WorldContext = "WorldContextObject", AdvancedDisplay = 1))
@@ -187,8 +131,4 @@ public:
 	/** Queries if the given audio bus is active (and audio can be mixed to it). */
 	UFUNCTION(BlueprintCallable, Category = "Audio|Bus", meta = (WorldContext = "WorldContextObject"))
 	static bool IsAudioBusActive(const UObject* WorldContextObject, UAudioBus* AudioBus);
-
-
-private:
-	static void PopulateSpectrumAnalyzerSettings(EFFTSize FFTSize, EFFTPeakInterpolationMethod InterpolationMethod, EFFTWindowType WindowType, float HopSize, Audio::FSpectrumAnalyzerSettings &OutSettings);
 };
diff --git a/Engine/Source/Runtime/AudioMixer/Public/AudioMixerDevice.h b/Engine/Source/Runtime/AudioMixer/Public/AudioMixerDevice.h
--- a/Engine/Source/Runtime/AudioMixer/Public/AudioMixerDevice.h
+++ b/Engine/Source/Runtime/AudioMixer/Public/AudioMixerDevice.h
@@ -130,6 +130,7 @@ namespace Audio
 		virtual void StopSpectrumAnalysis(USoundSubmix* InSubmix) override;
 		virtual void GetMagnitudesForFrequencies(USoundSubmix* InSubmix, const TArray<float>& InFrequencies, TArray<float>& OutMagnitudes);
 		virtual void GetPhasesForFrequencies(USoundSubmix* InSubmix, const TArray<float>& InFrequencies, TArray<float>& OutPhases);
+		virtual void AddSpectralAnalysisDelegate(USoundSubmix* InSubmix, const TArray<FSoundSubmixSpectralAnalysisBandSettings>& InBandSettings, const FOnSubmixSpectralAnalysisBP& OnSubmixSpectralAnalysisBP, float UpdateRate) override;
 
 		// Submix buffer listener callbacks
 		virtual void RegisterSubmixBufferListener(ISubmixBufferListener* InSubmixBufferListener, USoundSubmix* InSubmix = nullptr) override;
@@ -338,7 +339,7 @@ namespace Audio
 		FCriticalSection EndpointSubmixesMutationLock;
 
 		/** Which submixes have been told to envelope follow with this audio device. */
-		TArray<USoundSubmix*> EnvelopeFollowingSubmixes;
+		TArray<USoundSubmix*> DelegateBoundSubmixes;
 
 		/** Queue of mixer source voices. */
 		TQueue<FMixerSourceVoice*> SourceVoices;
diff --git a/Engine/Source/Runtime/AudioMixer/Public/AudioMixerSubmix.h b/Engine/Source/Runtime/AudioMixer/Public/AudioMixerSubmix.h
--- a/Engine/Source/Runtime/AudioMixer/Public/AudioMixerSubmix.h
+++ b/Engine/Source/Runtime/AudioMixer/Public/AudioMixerSubmix.h
@@ -200,6 +200,9 @@ namespace Audio
 		// Terminates whatever FFT Analyzer is being used for this submix.
 		void StopSpectrumAnalysis();
 
+		// Adds an envelope follower delegate
+		void AddSpectralAnalysisDelegate(const FOnSubmixSpectralAnalysisBP& OnSubmixSpectralAnalysisBP, const TArray<FSoundSubmixSpectralAnalysisBandSettings>& InBandSettings, float UpdateRate);
+
 		// Gets the most recent magnitude values for each corresponding value in InFrequencies (in Hz).
 		// This requires StartSpectrumAnalysis to be called first.
 		void GetMagnitudeForFrequencies(const TArray<float>& InFrequencies, TArray<float>& OutMagnitudes);
@@ -208,8 +211,8 @@ namespace Audio
 		// This requires StartSpectrumAnalysis to be called first.
 		void GetPhaseForFrequencies(const TArray<float>& InFrequencies, TArray<float>& OutPhases);
 
-		// Broadcast the envelope value on the game thread
-		void BroadcastEnvelope();
+		// Broadcast the envelope and submix delegates on the game thread
+		void BroadcastDelegates();
 
 		// returns true if this submix is encoded to a soundfield.
 		bool IsSoundfieldSubmix() const;
@@ -481,6 +484,42 @@ namespace Audio
 		// Multi-cast delegate to broadcast envelope data from this submix instance
 		FOnSubmixEnvelope OnSubmixEnvelope;
 
+		struct FSpectralAnalysisBandInfo
+		{
+			FEnvelopeFollower EnvelopeFollower;
+			float BandFrequency = 0.0f;
+		};
+
+		struct FSpectrumAnalysisDelegateInfo
+		{
+			FOnSubmixSpectralAnalysis OnSubmixSpectralAnalysis;
+			TUniquePtr<ISpectrumBandExtractor> SpectrumBandExtractor;
+			TArray<FSpectralAnalysisBandInfo> SpectralBands;
+
+			float LastUpdateTime = -1.0f;
+			float UpdateDelta = 0.0f;
+
+			FSpectrumAnalysisDelegateInfo()
+			{
+			}
+
+			FSpectrumAnalysisDelegateInfo(FSpectrumAnalysisDelegateInfo& Other)
+			{
+				OnSubmixSpectralAnalysis = Other.OnSubmixSpectralAnalysis;
+				SpectrumBandExtractor.Reset(Other.SpectrumBandExtractor.Release());
+				SpectralBands = Other.SpectralBands;
+			}
+
+			~FSpectrumAnalysisDelegateInfo()
+			{
+			}
+		};
+
+		TArray<FSpectrumAnalysisDelegateInfo> SpectralAnalysisDelegates;
+
+		// Bool set to true when spectrum analysis is enabled
+		FThreadSafeBool bIsSpectrumAnalyzing;
+
 		// Critical section used for when we are appending recorded data.
 		FCriticalSection RecordingCriticalSection;
 
diff --git a/Engine/Source/Runtime/Engine/Classes/Sound/SoundSubmix.h b/Engine/Source/Runtime/Engine/Classes/Sound/SoundSubmix.h
--- a/Engine/Source/Runtime/Engine/Classes/Sound/SoundSubmix.h
+++ b/Engine/Source/Runtime/Engine/Classes/Sound/SoundSubmix.h
@@ -11,6 +11,7 @@
 #include "SoundSubmixSend.h"
 #include "UObject/Object.h"
 #include "UObject/ObjectMacros.h"
+#include "DSP/SpectrumAnalyzer.h"
 
 #include "SoundSubmix.generated.h"
 
@@ -34,6 +35,72 @@ DECLARE_DYNAMIC_MULTICAST_DELEGATE_OneParam(FOnSubmixRecordedFileDone, const USo
 */
 DECLARE_DYNAMIC_MULTICAST_DELEGATE_OneParam(FOnSubmixEnvelope, const TArray<float>&, Envelope);
 
+DECLARE_DYNAMIC_MULTICAST_DELEGATE_OneParam(FOnSubmixSpectralAnalysis, const TArray<float>&, Magnitudes);
+
+UENUM(BlueprintType)
+enum class EFFTSize : uint8
+{
+	// 512
+	DefaultSize,
+
+	// 64
+	Min,
+
+	// 256
+	Small,
+
+	// 512
+	Medium,
+
+	// 1024
+	Large,
+
+	// 2048
+	VeryLarge, 
+
+	// 4096
+	Max
+};
+
+UENUM()
+enum class EFFTPeakInterpolationMethod : uint8
+{
+	NearestNeighbor,
+	Linear,
+	Quadratic,
+	ConstantQ,
+};
+
+UENUM()
+enum class EFFTWindowType : uint8
+{
+	// No window is applied. Technically a boxcar window.
+	None,
+
+	// Mainlobe width of -3 dB and sidelove attenuation of ~-40 dB. Good for COLA.
+	Hamming,
+
+	// Mainlobe width of -3 dB and sidelobe attenuation of ~-30dB. Good for COLA.
+	Hann,
+
+	// Mainlobe width of -3 dB and sidelobe attenuation of ~-60db. Tricky for COLA.
+	Blackman
+};
+
+UENUM(BlueprintType)
+enum class EAudioSpectrumType : uint8
+{
+	// Spectrum frequency values are equal to magnitude of frequency.
+	MagnitudeSpectrum,
+
+	// Spectrum frequency values are equal to magnitude squared.
+	PowerSpectrum,
+
+	// Returns decibels (0.0 dB is 1.0)
+	Decibel,
+};
+
+
 
 #if WITH_EDITOR
 
@@ -138,7 +205,7 @@ enum class EGainParamMode : uint8
 /**
  * Sound Submix class meant for applying an effect to the downmixed sum of multiple audio sources.
  */
-UCLASS(config = Engine, hidecategories = Object, editinlinenew, BlueprintType, Meta=(DisplayName="Effect Submix"))
+UCLASS(config = Engine, hidecategories = Object, editinlinenew, BlueprintType)
 class ENGINE_API USoundSubmix : public USoundSubmixWithParentBase
 {
 	GENERATED_UCLASS_BODY()
@@ -222,13 +289,39 @@ public:
 
 	void StopEnvelopeFollowing(FAudioDevice* InDevice);
 
+	/**
+	 *	Adds an envelope follower delegate to the submix when envelope following is enabled on this submix.
+	 *	@param	OnSubmixEnvelopeBP	Event to fire when new envelope data is available.
+	 */
 	UFUNCTION(BlueprintCallable, Category = "Audio|EnvelopeFollowing", meta = (WorldContext = "WorldContextObject"))
 	void AddEnvelopeFollowerDelegate(const UObject* WorldContextObject, const FOnSubmixEnvelopeBP& OnSubmixEnvelopeBP);
 
+	/**
+	 *	Adds a spectral envelope delegate to receive notifications when this submix has spectrum analysis enabled.
+	 *	@param	InBandsettings					The frequency bands to analyze and their envelope-following settings.
+	 *  @param  OnSubmixSpectralAnalysisBP		Event to fire when new spectral data is available.
+	 *	@param	UpdateRate						How often to retrieve the data from the spectral analyzer and broadcast the event. Max is 30 times per second.
+	 */
+	UFUNCTION(BlueprintCallable, Category = "Audio|Spectrum", meta = (WorldContext = "WorldContextObject"))
+	void AddSpectralAnalysisDelegate(const UObject* WorldContextObject, const TArray<FSoundSubmixSpectralAnalysisBandSettings>& InBandSettings, const FOnSubmixSpectralAnalysisBP& OnSubmixSpectralAnalysisBP, float UpdateRate);
+
+	/** Start spectrum analysis of the audio output. */
+	UFUNCTION(BlueprintCallable, Category = "Audio|Analysis", meta = (WorldContext = "WorldContextObject", AdvancedDisplay = 1))
+	void StartSpectralAnalysis(const UObject* WorldContextObject, EFFTSize FFTSize = EFFTSize::DefaultSize, EFFTPeakInterpolationMethod InterpolationMethod = EFFTPeakInterpolationMethod::Linear, EFFTWindowType WindowType = EFFTWindowType::Hann, float HopSize = 0, EAudioSpectrumType SpectrumType = EAudioSpectrumType::MagnitudeSpectrum);
+
+	void StartSpectralAnalysis(FAudioDevice* InDevice, EFFTSize FFTSize = EFFTSize::DefaultSize, EFFTPeakInterpolationMethod InterpolationMethod = EFFTPeakInterpolationMethod::Linear, EFFTWindowType WindowType = EFFTWindowType::Hann, float HopSize = 0, EAudioSpectrumType SpectrumType = EAudioSpectrumType::MagnitudeSpectrum);
+
+	/** Start spectrum analysis of the audio output. */
+	UFUNCTION(BlueprintCallable, Category = "Audio|Analysis", meta = (WorldContext = "WorldContextObject", AdvancedDisplay = 1))
+	void StopSpectralAnalysis(const UObject* WorldContextObject);
+
+	void StopSpectralAnalysis(FAudioDevice* InDevice);
+
 	/** Sets the output volume of the submix. This dynamic volume scales with the OutputVolume property of this submix. */
 	UFUNCTION(BlueprintCallable, Category = "Audio", meta = (WorldContext = "WorldContextObject"))
 	void SetSubmixOutputVolume(const UObject* WorldContextObject, float InOutputVolume);
 
+	static Audio::FSpectrumAnalyzerSettings GetSpectrumAnalyzerSettings(EFFTSize FFTSize, EFFTPeakInterpolationMethod InterpolationMethod, EFFTWindowType WindowType, float HopSize, EAudioSpectrumType AudioSpectrumType);
 protected:
 
 	virtual void PostLoad() override;
@@ -244,7 +337,7 @@ protected:
 /**
  * Sound Submix class meant for use with soundfield formats, such as Ambisonics.
  */
-UCLASS(config = Engine, hidecategories = Object, editinlinenew, BlueprintType, Meta = (DisplayName = "Soundfield Submix"))
+UCLASS(config = Engine, hidecategories = Object, editinlinenew, BlueprintType, Meta = (DisplayName = "Sound Submix Soundfield"))
 class ENGINE_API USoundfieldSubmix : public USoundSubmixWithParentBase
 {
 	GENERATED_UCLASS_BODY()
@@ -291,7 +384,7 @@ protected:
 /**
  * Sound Submix class meant for sending audio to an external endpoint, such as controller haptics or an additional audio device.
  */
-UCLASS(config = Engine, hidecategories = Object, editinlinenew, BlueprintType, Meta = (DisplayName = "Audio Endpoint Submix"))
+UCLASS(config = Engine, hidecategories = Object, editinlinenew, BlueprintType, Meta = (DisplayName = "Sound Submix Endpoint"))
 class ENGINE_API UEndpointSubmix : public USoundSubmixBase
 {
 	GENERATED_UCLASS_BODY()
@@ -318,7 +411,7 @@ public:
 /**
  * Sound Submix class meant for sending soundfield-encoded audio to an external endpoint, such as a hardware binaural renderer that supports ambisonics.
  */
-UCLASS(config = Engine, hidecategories = Object, editinlinenew, BlueprintType, Meta = (DisplayName = "Soundfield Endpoint Submix"))
+UCLASS(config = Engine, hidecategories = Object, editinlinenew, BlueprintType, Meta = (DisplayName = "Sound Submix Soundfield Endpoint"))
 class ENGINE_API USoundfieldEndpointSubmix : public USoundSubmixBase
 {
 	GENERATED_UCLASS_BODY()
diff --git a/Engine/Source/Runtime/Engine/Classes/Sound/SoundSubmixSend.h b/Engine/Source/Runtime/Engine/Classes/Sound/SoundSubmixSend.h
--- a/Engine/Source/Runtime/Engine/Classes/Sound/SoundSubmixSend.h
+++ b/Engine/Source/Runtime/Engine/Classes/Sound/SoundSubmixSend.h
@@ -13,9 +13,29 @@
 // Forward Declarations
 class USoundSubmixBase;
 
+USTRUCT(BlueprintType)
+struct ENGINE_API FSoundSubmixSpectralAnalysisBandSettings
+{
+	GENERATED_USTRUCT_BODY()
+
+	// The frequency band for the magnitude to analyze
+	UPROPERTY(EditAnywhere, BlueprintReadWrite, Category = SubmixSpectralAnalysis, meta = (ClampMin = "10.0", ClampMax = "20000.0", UIMin = "10.0", UIMax = "10000.0"))
+	float BandFrequency = 440.0f;
+
+	// The attack time for the FFT band interpolation for delegate callback
+	UPROPERTY(EditAnywhere, BlueprintReadWrite, Category = SubmixSpectralAnalysis, meta = (ClampMin = "0.0", UIMin = "10.0", UIMax = "10000.0"))
+	int32 AttackTimeMsec = 10;
+
+	// The release time for the FFT band interpolation for delegate callback
+	UPROPERTY(EditAnywhere, BlueprintReadWrite, Category = SubmixSpectralAnalysis, meta = (ClampMin = "0.0", UIMin = "10.0", UIMax = "10000.0"))
+	int32 ReleaseTimeMsec = 500;
+};
 
 DECLARE_DYNAMIC_DELEGATE_OneParam(FOnSubmixEnvelopeBP, const TArray<float>&, Envelope);
 
+DECLARE_DYNAMIC_DELEGATE_OneParam(FOnSubmixSpectralAnalysisBP, const TArray<float>&, Magnitude);
+
+
 UENUM(BlueprintType)
 enum class EAudioRecordingExportType : uint8
 {
@@ -52,15 +72,15 @@ struct ENGINE_API FSoundSubmixSendInfo
 		Linear: Interpolate between Min and Max Send Levels based on listener distance (between Distance Min and Distance Max)
 		Custom Curve: Use the float curve to map Send Level to distance (0.0-1.0 on curve maps to Distance Min - Distance Max)
 	*/
-	UPROPERTY(EditAnywhere, Category = SubmixSend)
+	UPROPERTY(EditAnywhere, BlueprintReadWrite, Category = SubmixSend)
 	ESendLevelControlMethod SendLevelControlMethod;
 
 	// The submix to send the audio to
-	UPROPERTY(EditAnywhere, Category = SubmixSend)
+	UPROPERTY(EditAnywhere, BlueprintReadWrite, Category = SubmixSend)
 	USoundSubmixBase* SoundSubmix;
 
 	// The amount of audio to send
-	UPROPERTY(EditAnywhere, Category = SubmixSend)
+	UPROPERTY(EditAnywhere, BlueprintReadWrite, Category = SubmixSend)
 	float SendLevel;
 
 	// The amount to send to master when sound is located at a distance equal to value specified in the min send distance.
diff --git a/Engine/Source/Runtime/Engine/Private/AudioDevice.cpp b/Engine/Source/Runtime/Engine/Private/AudioDevice.cpp
--- a/Engine/Source/Runtime/Engine/Private/AudioDevice.cpp
+++ b/Engine/Source/Runtime/Engine/Private/AudioDevice.cpp
@@ -4695,6 +4695,11 @@ void FAudioDevice::GetPhasesForFrequencies(USoundSubmix* InSubmix, const TArray<
 	UE_LOG(LogAudio, Error, TEXT("Spectrum analysis of submixes only works with the audio mixer. Please run using -audiomixer or set INI file to use submix recording."));
 }
 
+void FAudioDevice::AddSpectralAnalysisDelegate(USoundSubmix* InSubmix, const TArray<FSoundSubmixSpectralAnalysisBandSettings>& InBandSettings, const FOnSubmixSpectralAnalysisBP& OnSubmixSpectralAnalysisBP, float UpdateRate)
+{
+	UE_LOG(LogAudio, Error, TEXT("Spectrum analysis of submixes only works with the audio mixer. Please run using -audiomixer or set INI file to use submix recording."));
+}
+
 void FAudioDevice::AddVirtualLoop(const FAudioVirtualLoop& InVirtualLoop)
 {
 	FAudioVirtualLoop VirtualLoop = InVirtualLoop;
diff --git a/Engine/Source/Runtime/Engine/Private/SoundSubmix.cpp b/Engine/Source/Runtime/Engine/Private/SoundSubmix.cpp
--- a/Engine/Source/Runtime/Engine/Private/SoundSubmix.cpp
+++ b/Engine/Source/Runtime/Engine/Private/SoundSubmix.cpp
@@ -8,6 +8,7 @@
 #include "Sound/SoundSubmixSend.h"
 #include "UObject/UObjectIterator.h"
 #include "DSP/Dsp.h"
+#include "DSP/SpectrumAnalyzer.h"
 
 #if WITH_EDITOR
 #include "Framework/Notifications/NotificationManager.h"
@@ -91,10 +92,13 @@ void USoundSubmix::StartRecordingOutput(const UObject* WorldContextObject, float
 	}
 
 	// Find device for this specific audio recording thing.
-	UWorld* ThisWorld = GEngine->GetWorldFromContextObject(WorldContextObject, EGetWorldErrorMode::LogAndReturnNull);
-	FAudioDevice* DesiredAudioDevice = ThisWorld->GetAudioDeviceRaw();
-
-	StartRecordingOutput(DesiredAudioDevice, ExpectedDuration);
+	if (UWorld* ThisWorld = GEngine->GetWorldFromContextObject(WorldContextObject, EGetWorldErrorMode::LogAndReturnNull))
+	{
+		if (FAudioDevice* AudioDevice = ThisWorld->GetAudioDeviceRaw())
+		{
+			StartRecordingOutput(AudioDevice, ExpectedDuration);
+		}
+	}
 }
 
 void USoundSubmix::StartRecordingOutput(FAudioDevice* InDevice, float ExpectedDuration)
@@ -113,10 +117,13 @@ void USoundSubmix::StopRecordingOutput(const UObject* WorldContextObject, EAudio
 	}
 
 	// Find device for this specific audio recording thing.
-	UWorld* ThisWorld = GEngine->GetWorldFromContextObject(WorldContextObject, EGetWorldErrorMode::LogAndReturnNull);
-	FAudioDevice* DesiredAudioDevice = ThisWorld->GetAudioDeviceRaw();
-
-	StopRecordingOutput(DesiredAudioDevice, ExportType, Name, Path, ExistingSoundWaveToOverwrite);
+	if (UWorld* ThisWorld = GEngine->GetWorldFromContextObject(WorldContextObject, EGetWorldErrorMode::LogAndReturnNull))
+	{
+		if (FAudioDevice* AudioDevice = ThisWorld->GetAudioDeviceRaw())
+		{
+			StopRecordingOutput(AudioDevice, ExportType, Name, Path, ExistingSoundWaveToOverwrite);
+		}
+	}
 }
 
 void USoundSubmix::StopRecordingOutput(FAudioDevice* InDevice, EAudioRecordingExportType ExportType, const FString& Name, FString Path, USoundWave* ExistingSoundWaveToOverwrite /*= nullptr*/)
@@ -193,10 +200,13 @@ void USoundSubmix::StartEnvelopeFollowing(const UObject* WorldContextObject)
 	}
 
 	// Find device for this specific audio recording thing.
-	UWorld* ThisWorld = GEngine->GetWorldFromContextObject(WorldContextObject, EGetWorldErrorMode::LogAndReturnNull);
-	FAudioDevice* AudioDevice = ThisWorld->GetAudioDeviceRaw();
-
-	StartEnvelopeFollowing(AudioDevice);
+	if (UWorld* ThisWorld = GEngine->GetWorldFromContextObject(WorldContextObject, EGetWorldErrorMode::LogAndReturnNull))
+	{
+		if (FAudioDevice* AudioDevice = ThisWorld->GetAudioDeviceRaw())
+		{
+			StartEnvelopeFollowing(AudioDevice);
+		}
+	}
 }
 
 void USoundSubmix::StartEnvelopeFollowing(FAudioDevice* InAudioDevice)
@@ -215,10 +225,13 @@ void USoundSubmix::StopEnvelopeFollowing(const UObject* WorldContextObject)
 	}
 
 	// Find device for this specific audio recording thing.
-	UWorld* ThisWorld = GEngine->GetWorldFromContextObject(WorldContextObject, EGetWorldErrorMode::LogAndReturnNull);
-	FAudioDevice* AudioDevice = ThisWorld->GetAudioDeviceRaw();
-
-	StopEnvelopeFollowing(AudioDevice);
+	if (UWorld* ThisWorld = GEngine->GetWorldFromContextObject(WorldContextObject, EGetWorldErrorMode::LogAndReturnNull))
+	{
+		if (FAudioDevice* AudioDevice = ThisWorld->GetAudioDeviceRaw())
+		{
+			StopEnvelopeFollowing(AudioDevice);
+		}
+	}
 }
 
 void USoundSubmix::StopEnvelopeFollowing(FAudioDevice* InAudioDevice)
@@ -237,11 +250,85 @@ void USoundSubmix::AddEnvelopeFollowerDelegate(const UObject* WorldContextObject
 	}
 
 	// Find device for this specific audio recording thing.
-	UWorld* ThisWorld = GEngine->GetWorldFromContextObject(WorldContextObject, EGetWorldErrorMode::LogAndReturnNull);
-	FAudioDevice* AudioDevice = ThisWorld->GetAudioDeviceRaw();
-	if (AudioDevice)
+	if (UWorld* ThisWorld = GEngine->GetWorldFromContextObject(WorldContextObject, EGetWorldErrorMode::LogAndReturnNull))
+	{
+		if (FAudioDevice* AudioDevice = ThisWorld->GetAudioDeviceRaw())
+		{
+			AudioDevice->AddEnvelopeFollowerDelegate(this, OnSubmixEnvelopeBP);
+		}
+	}
+}
+
+void USoundSubmix::AddSpectralAnalysisDelegate(const UObject* WorldContextObject, const TArray<FSoundSubmixSpectralAnalysisBandSettings>& InBandSettings, const FOnSubmixSpectralAnalysisBP& OnSubmixSpectralAnalysisBP, float UpdateRate)
+{
+	if (!GEngine)
+	{
+		return;
+	}
+
+	// Find device for this specific audio recording thing.
+	if (UWorld* ThisWorld = GEngine->GetWorldFromContextObject(WorldContextObject, EGetWorldErrorMode::LogAndReturnNull))
+	{
+		if (FAudioDevice* AudioDevice = ThisWorld->GetAudioDeviceRaw())
+		{
+			AudioDevice->AddSpectralAnalysisDelegate(this, InBandSettings, OnSubmixSpectralAnalysisBP, UpdateRate);
+		}
+	}
+}
+
+void USoundSubmix::StartSpectralAnalysis(const UObject* WorldContextObject, EFFTSize FFTSize, EFFTPeakInterpolationMethod InterpolationMethod, EFFTWindowType WindowType, float HopSize, EAudioSpectrumType SpectrumType)
+{
+	if (!GEngine)
+	{
+		return;
+	}
+
+	// Find device for this specific audio recording thing.
+	if (UWorld* ThisWorld = GEngine->GetWorldFromContextObject(WorldContextObject, EGetWorldErrorMode::LogAndReturnNull))
+	{	
+		if (FAudioDevice* AudioDevice = ThisWorld->GetAudioDeviceRaw())
+		{
+			StartSpectralAnalysis(AudioDevice, FFTSize, InterpolationMethod, WindowType, HopSize, SpectrumType);
+		}
+	}
+}
+
+void USoundSubmix::StartSpectralAnalysis(FAudioDevice* InAudioDevice, EFFTSize FFTSize, EFFTPeakInterpolationMethod InterpolationMethod, EFFTWindowType WindowType, float HopSize, EAudioSpectrumType SpectrumType)
+{
+	if (!GEngine)
+	{
+		return;
+	}
+
+	if (InAudioDevice)
+	{
+		Audio::FSpectrumAnalyzerSettings Settings = USoundSubmix::GetSpectrumAnalyzerSettings(FFTSize, InterpolationMethod, WindowType, HopSize, SpectrumType);
+		InAudioDevice->StartSpectrumAnalysis(this, Settings);
+	}
+}
+
+void USoundSubmix::StopSpectralAnalysis(const UObject* WorldContextObject)
+{
+	if (!GEngine)
+	{
+		return;
+	}
+
+	// Find device for this specific audio recording thing.
+	if (UWorld* ThisWorld = GEngine->GetWorldFromContextObject(WorldContextObject, EGetWorldErrorMode::LogAndReturnNull))
 	{
-		AudioDevice->AddEnvelopeFollowerDelegate(this, OnSubmixEnvelopeBP);
+		if (FAudioDevice* AudioDevice = ThisWorld->GetAudioDeviceRaw())
+		{
+			AudioDevice->StopSpectrumAnalysis(this);
+		}
+	}
+}
+
+void USoundSubmix::StopSpectralAnalysis(FAudioDevice* InAudioDevice)
+{
+	if (InAudioDevice)
+	{
+		InAudioDevice->StopSpectrumAnalysis(this);
 	}
 }
 
@@ -252,11 +339,12 @@ void USoundSubmix::SetSubmixOutputVolume(const UObject* WorldContextObject, floa
 		return;
 	}
 
-	UWorld* ThisWorld = GEngine->GetWorldFromContextObject(WorldContextObject, EGetWorldErrorMode::LogAndReturnNull);
-	FAudioDevice* AudioDevice = ThisWorld->GetAudioDeviceRaw();
-	if (AudioDevice)
+	if (UWorld* ThisWorld = GEngine->GetWorldFromContextObject(WorldContextObject, EGetWorldErrorMode::LogAndReturnNull))
 	{
-		AudioDevice->SetSubmixOutputVolume(this, InOutputVolume);
+		if (FAudioDevice* AudioDevice = ThisWorld->GetAudioDeviceRaw())
+		{
+			AudioDevice->SetSubmixOutputVolume(this, InOutputVolume);
+		}
 	}
 }
 
@@ -333,8 +421,6 @@ void USoundSubmix::PostEditChangeProperty(struct FPropertyChangedEvent& Property
 		{
 			if (bUpdateSubmixGain)
 			{
-//				AudioDeviceManager->UpdateSubmix(this);
-
 				const float NewOutputVolume = OutputVolume;
 				const float NewWetLevel = WetLevel;
 				const float NewDryLevel = DryLevel;
@@ -781,6 +867,91 @@ TArray<USoundfieldEffectBase*> USoundfieldEndpointSubmix::GetSoundfieldProcessor
 	return SoundfieldEffectChain;
 }
 
+
+Audio::FSpectrumAnalyzerSettings USoundSubmix::GetSpectrumAnalyzerSettings(EFFTSize FFTSize, EFFTPeakInterpolationMethod InterpolationMethod, EFFTWindowType WindowType, float HopSize, EAudioSpectrumType AudioSpectrumType)
+{
+	Audio::FSpectrumAnalyzerSettings OutSettings;
+
+	switch (FFTSize)
+	{
+	case EFFTSize::DefaultSize:
+		OutSettings.FFTSize = Audio::FSpectrumAnalyzerSettings::EFFTSize::Default;
+		break;
+	case EFFTSize::Min:
+		OutSettings.FFTSize = Audio::FSpectrumAnalyzerSettings::EFFTSize::Min_64;
+		break;
+	case EFFTSize::Small:
+		OutSettings.FFTSize = Audio::FSpectrumAnalyzerSettings::EFFTSize::Small_256;
+		break;
+	case EFFTSize::Medium:
+		OutSettings.FFTSize = Audio::FSpectrumAnalyzerSettings::EFFTSize::Medium_512;
+		break;
+	case EFFTSize::Large:
+		OutSettings.FFTSize = Audio::FSpectrumAnalyzerSettings::EFFTSize::Large_1024;
+		break;
+	case EFFTSize::VeryLarge:
+		OutSettings.FFTSize = Audio::FSpectrumAnalyzerSettings::EFFTSize::VeryLarge_2048;
+		break;
+	case EFFTSize::Max:
+		OutSettings.FFTSize = Audio::FSpectrumAnalyzerSettings::EFFTSize::TestLarge_4096;
+		break;
+	default:
+		break;
+	}
+
+	switch (InterpolationMethod)
+	{
+	case EFFTPeakInterpolationMethod::NearestNeighbor:
+		OutSettings.InterpolationMethod = Audio::FSpectrumAnalyzerSettings::EPeakInterpolationMethod::NearestNeighbor;
+		break;
+	case EFFTPeakInterpolationMethod::Linear:
+		OutSettings.InterpolationMethod = Audio::FSpectrumAnalyzerSettings::EPeakInterpolationMethod::Linear;
+		break;
+	case EFFTPeakInterpolationMethod::Quadratic:
+		OutSettings.InterpolationMethod = Audio::FSpectrumAnalyzerSettings::EPeakInterpolationMethod::Quadratic;
+		break;
+	case EFFTPeakInterpolationMethod::ConstantQ:
+		OutSettings.InterpolationMethod = Audio::FSpectrumAnalyzerSettings::EPeakInterpolationMethod::ConstantQ;
+	default:
+		break;
+	}
+
+	switch (WindowType)
+	{
+	case EFFTWindowType::None:
+		OutSettings.WindowType = Audio::EWindowType::None;
+		break;
+	case EFFTWindowType::Hamming:
+		OutSettings.WindowType = Audio::EWindowType::Hamming;
+		break;
+	case EFFTWindowType::Hann:
+		OutSettings.WindowType = Audio::EWindowType::Hann;
+		break;
+	case EFFTWindowType::Blackman:
+		OutSettings.WindowType = Audio::EWindowType::Blackman;
+		break;
+	default:
+		break;
+	}
+
+	switch (AudioSpectrumType)
+	{
+	case EAudioSpectrumType::MagnitudeSpectrum:
+		OutSettings.SpectrumType = Audio::FSpectrumAnalyzerSettings::ESpectrumAnalyzerType::Magnitude;
+		break;
+	case EAudioSpectrumType::PowerSpectrum:
+		OutSettings.SpectrumType = Audio::FSpectrumAnalyzerSettings::ESpectrumAnalyzerType::Power;
+		break;
+	case EAudioSpectrumType::Decibel:
+		OutSettings.SpectrumType = Audio::FSpectrumAnalyzerSettings::ESpectrumAnalyzerType::Decibel;
+		break;
+	}
+
+	OutSettings.HopSize = HopSize;
+
+	return OutSettings;
+}
+
 void USoundfieldEndpointSubmix::SanitizeLinks()
 {
 	bool bShouldRefreshEditor = false;
diff --git a/Engine/Source/Runtime/Engine/Public/AudioDevice.h b/Engine/Source/Runtime/Engine/Public/AudioDevice.h
--- a/Engine/Source/Runtime/Engine/Public/AudioDevice.h
+++ b/Engine/Source/Runtime/Engine/Public/AudioDevice.h
@@ -1308,6 +1308,8 @@ public:
 	virtual void StopSpectrumAnalysis(USoundSubmix* InSubmix);
 	virtual void GetMagnitudesForFrequencies(USoundSubmix* InSubmix, const TArray<float>& InFrequencies, TArray<float>& OutMagnitudes);
 	virtual void GetPhasesForFrequencies(USoundSubmix* InSubmix, const TArray<float>& InFrequencies, TArray<float>& OutPhases);
+	virtual void AddSpectralAnalysisDelegate(USoundSubmix* InSubmix, const TArray<FSoundSubmixSpectralAnalysisBandSettings>& InBandSettings, const FOnSubmixSpectralAnalysisBP& OnSubmixSpectralAnalysisBP, float UpdateRate);
+
 
 protected:
 	friend class FSoundSource;
diff --git a/Engine/Source/Runtime/SignalProcessing/Private/AudioFFT.cpp b/Engine/Source/Runtime/SignalProcessing/Private/AudioFFT.cpp
--- a/Engine/Source/Runtime/SignalProcessing/Private/AudioFFT.cpp
+++ b/Engine/Source/Runtime/SignalProcessing/Private/AudioFFT.cpp
@@ -25,7 +25,7 @@ namespace Audio
 
 		for (int32 FrameIndex = 0; FrameIndex < NumFrames; FrameIndex++)
 		{
-			const float Value = 0.54 -  0.46f * (1 - FMath::Cos(Phase));
+			const float Value = 0.54 -  (0.46f * FMath::Cos(Phase));
 			Phase += PhaseDelta;
 
 			for (int32 ChannelIndex = 0; ChannelIndex < NumChannels; ChannelIndex++)
diff --git a/Engine/Source/Runtime/SignalProcessing/Private/ConstantQ.cpp b/Engine/Source/Runtime/SignalProcessing/Private/ConstantQ.cpp
--- a/Engine/Source/Runtime/SignalProcessing/Private/ConstantQ.cpp
+++ b/Engine/Source/Runtime/SignalProcessing/Private/ConstantQ.cpp
@@ -4,58 +4,158 @@
 #include "CoreMinimal.h"
 #include "DSP/FloatArrayMath.h"
 
-namespace
+namespace Audio
 {
-	float GetConstantQCenterFrequency(const int32 InBandIndex, const float InBaseFrequency, const float InBandsPerOctave)
+	namespace PseudoConstantQIntrinsics
+	{
+		void FillArrayWithTruncatedGaussian(const float InCenterFreq, const float InBandWidth, const int32 InFFTSize, const float InSampleRate, AlignedFloatBuffer& OutOffsetArray, int32& OutOffsetIndex)
+		{
+			check(InBandWidth > 0.f);
+			check(InFFTSize > 0);
+			check(InCenterFreq >= 0.f);
+			check(InSampleRate > 0.f);
+
+			
+			// Determine points where gaussian will value will go below a small number
+			const float SignificantHalfBandWidth = InBandWidth * FMath::Sqrt(-2.f * FMath::Loge(SMALL_NUMBER));
+			const float LowestSignificantFreq = FMath::Clamp(InCenterFreq - SignificantHalfBandWidth, 0.f, InSampleRate / 2.f);
+			const float HighestSignificantFreq = FMath::Clamp(InCenterFreq + SignificantHalfBandWidth, 0.f, InSampleRate / 2.f);
+			int32 LowestSignificantIndex = FMath::CeilToInt(InFFTSize * LowestSignificantFreq / InSampleRate);
+			int32 HighestSignificantIndex = FMath::FloorToInt(InFFTSize * HighestSignificantFreq / InSampleRate);
+			int32 Num = HighestSignificantIndex - LowestSignificantIndex + 1;
+
+			if (Num < 1)
+			{
+				Num = 1;
+			}
+
+			// Prepare outputs
+			OutOffsetIndex = LowestSignificantIndex;
+			OutOffsetArray.Reset(Num);
+			OutOffsetArray.AddUninitialized(Num);
+
+			// Fill array with truncated Gaussian
+			const float BandWidthSquared = InBandWidth * InBandWidth;
+			for (int32 i = 0; i < Num; i++)
+			{
+				float FFTBinHz = static_cast<float>(i + LowestSignificantIndex) * InSampleRate / InFFTSize;
+				float DeltaHz = FFTBinHz - InCenterFreq;
+				OutOffsetArray[i] = FMath::Exp(-0.5 * (DeltaHz * DeltaHz) / BandWidthSquared);
+			}
+
+			// set infs and nans to zero
+			for (int32 i = 0; i < OutOffsetArray.Num(); i++)
+			{
+				if (!FMath::IsFinite(OutOffsetArray[i]))
+				{
+					OutOffsetArray[i] = 0.f;
+				}
+			}
+		}
+	}
+
+
+	float FPseudoConstantQ::GetConstantQCenterFrequency(const int32 InBandIndex, const float InBaseFrequency, const float InBandsPerOctave)
 	{
 		check(InBandsPerOctave > 0.f);
 		return InBaseFrequency * FMath::Pow(2.f, static_cast<float>(InBandIndex) / InBandsPerOctave);
 	}
 
-	float GetConstantQBandWidth(const float InBandCenter, const float InBandsPerOctave, const float InBandWidthStretch)
+	float FPseudoConstantQ::GetConstantQBandWidth(const float InBandCenter, const float InBandsPerOctave, const float InBandWidthStretch)
 	{
 		check(InBandsPerOctave > 0.f);
 		return InBandWidthStretch * InBandCenter * (FMath::Pow(2.f, 1.f / InBandsPerOctave) - 1.f);
 	}
 
-	void FillArrayWithTruncatedGaussian(const float InCenterFreq, const float InBandWidth, const int32 InFFTSize, const float InSampleRate, TArray<float>& OutOffsetArray, int32& OutOffsetIndex)
+
+	void FPseudoConstantQ::FillArrayWithConstantQBand(const FPseudoConstantQBandSettings& InSettings, AlignedFloatBuffer& OutOffsetArray, int32& OutOffsetIndex)
 	{
-		check(InBandWidth > 0.f);
-		check(InFFTSize > 0);
-		check(InCenterFreq >= 0.f);
-		check(InSampleRate > 0.f);
+		check(InSettings.SampleRate > 0.f);
+		check(InSettings.FFTSize > 0);
+
+		const int32 NumUsefulFFTBins = (InSettings.FFTSize / 2) + 1;
+
+		// Create gaussian centered around center freq and with appropriate bandwidth
+		OutOffsetIndex = 0;
+		float BandWidth = FMath::Max(SMALL_NUMBER, InSettings.BandWidth);
+		PseudoConstantQIntrinsics::FillArrayWithTruncatedGaussian(InSettings.CenterFreq, BandWidth, InSettings.FFTSize, InSettings.SampleRate, OutOffsetArray, OutOffsetIndex);
 
-		// Determine points where gaussian will value will go below a small number
-		const float SignificantHalfBandWidth = InBandWidth * FMath::Sqrt(-2.f * FMath::Loge(SMALL_NUMBER));
-		const float LowestSignificantFreq = FMath::Clamp(InCenterFreq - SignificantHalfBandWidth, 0.f, InSampleRate / 2.f);
-		const float HighestSignificantFreq = FMath::Clamp(InCenterFreq + SignificantHalfBandWidth, 0.f, InSampleRate / 2.f);
-		int32 LowestSignificantIndex = FMath::CeilToInt(InFFTSize * LowestSignificantFreq / InSampleRate);
-		int32 HighestSignificantIndex = FMath::FloorToInt(InFFTSize * HighestSignificantFreq / InSampleRate);
- 		int32 Num = HighestSignificantIndex - LowestSignificantIndex + 1;
+		// Need to have a lower bound for 
+		float DigitalBandWidth = FMath::Max(1.f / static_cast<float>(InSettings.FFTSize), BandWidth / InSettings.SampleRate);
 
-		if (Num < 1)
+		// Sanity check the CQT bins to make sure we the bandwidth wasn't so small that the array is essentially empty
+		if (OutOffsetArray.Num() > 0)
 		{
-			Num = 1;
+			float GaussianSum = 0.f;
+			Audio::ArraySum(OutOffsetArray, GaussianSum);
+
+			if (GaussianSum < 0.95f)
+			{
+				// It's a bit of an arbitrary threshold, but it tells us that the bandwidth
+				// is low enough and the FFT granularity course enough to where our pseudo 
+				// cqt windows will likely miss data. In this case we force the window to 
+				// interpolate between two nearest fft bins.
+				FMemory::Memset(OutOffsetArray.GetData(), 0, OutOffsetArray.Num() * sizeof(float));
+				float Position = (InSettings.FFTSize * InSettings.CenterFreq / InSettings.SampleRate) - static_cast<float>(OutOffsetIndex);
+				int32 LowerIndex = FMath::FloorToInt(Position);
+				int32 UpperIndex = LowerIndex + 1;
+
+				if ((LowerIndex >= 0) && (UpperIndex < OutOffsetArray.Num()))
+				{
+					OutOffsetArray[LowerIndex] = 1.f - (Position - static_cast<float>(LowerIndex));
+					OutOffsetArray[UpperIndex] = 1.f - OutOffsetArray[LowerIndex];
+
+					// Need to update digital bandwidth to be bandwidth of FFTbin
+					DigitalBandWidth = 1.f / static_cast<float>(InSettings.FFTSize);
+				}
+				else if ((LowerIndex >= 0) && (LowerIndex < OutOffsetArray.Num()))
+				{
+					OutOffsetArray[LowerIndex] = 1.f;
+					DigitalBandWidth = 1.f / static_cast<float>(InSettings.FFTSize);
+				}
+				else if ((UpperIndex >= 0) && (UpperIndex < OutOffsetArray.Num()))
+				{
+					OutOffsetArray[UpperIndex] = 1.f;
+					DigitalBandWidth = 1.f / static_cast<float>(InSettings.FFTSize);
+				}
+			}
 		}
 
-		// Prepare outputs
-		OutOffsetIndex = LowestSignificantIndex;
-		OutOffsetArray.Reset(Num);
-		OutOffsetArray.AddUninitialized(Num);
+		// Normalize window
+		float NormDenom = 1.f;
+		switch (InSettings.Normalization)
+		{
+			case EPseudoConstantQNormalization::EqualAmplitude:
+				NormDenom = 1.f;
+				break;
+
+			case EPseudoConstantQNormalization::EqualEuclideanNorm:
+				NormDenom = FMath::Sqrt(DigitalBandWidth * InSettings.FFTSize * FMath::Sqrt(PI  * 2.f));
+				break;
+
+			case EPseudoConstantQNormalization::EqualEnergy:
+			default:
+				NormDenom = DigitalBandWidth * InSettings.FFTSize * FMath::Sqrt(PI  * 2.f);
+				break;
+		}
 
-		// Fill array with truncated Gaussian
-		const float BandWidthSquared = InBandWidth * InBandWidth;
-		for (int32 i = 0; i < Num; i++)
+		if ((NormDenom > 0.f) && (NormDenom != 1.f))
 		{
-			float FFTBinHz = static_cast<float>(i + LowestSignificantIndex) * InSampleRate / InFFTSize;
-			float DeltaHz = FFTBinHz - InCenterFreq;
-			OutOffsetArray[i] = FMath::Exp(-0.5 * (DeltaHz * DeltaHz) / BandWidthSquared);
+			ArrayMultiplyByConstantInPlace(OutOffsetArray, 1.f / NormDenom);
+		}
+
+		// Truncate if necessary
+		if (OutOffsetIndex >= NumUsefulFFTBins)
+		{
+			OutOffsetIndex = 0;
+			OutOffsetArray.Empty();
+		}
+		else if ((OutOffsetIndex + OutOffsetArray.Num()) > NumUsefulFFTBins)
+		{
+			OutOffsetArray.SetNum(NumUsefulFFTBins - OutOffsetIndex);
 		}
 	}
-}
 
-namespace Audio
-{
 	TUniquePtr<FContiguousSparse2DKernelTransform> NewPseudoConstantQKernelTransform(const FPseudoConstantQKernelSettings& InSettings, const int32 InFFTSize, const float InSampleRate)
 	{
 		check(InSampleRate > 0.f);
@@ -65,113 +165,32 @@ namespace Audio
 
 		TUniquePtr<FContiguousSparse2DKernelTransform> Transform = MakeUnique<FContiguousSparse2DKernelTransform>(NumUsefulFFTBins, InSettings.NumBands);
 
+		FPseudoConstantQBandSettings BandSettings;
+		BandSettings.FFTSize = InFFTSize;
+		BandSettings.SampleRate = InSampleRate;
+		BandSettings.Normalization = InSettings.Normalization;
+
 		for (int32 CQTBandIndex = 0; CQTBandIndex < InSettings.NumBands; CQTBandIndex++)
 		{
 			// Determine band center and width for this CQT band
-			const float BandCenter = GetConstantQCenterFrequency(CQTBandIndex, InSettings.KernelLowestCenterFreq, InSettings.NumBandsPerOctave);
-			const float BandWidth = GetConstantQBandWidth(BandCenter, InSettings.NumBandsPerOctave, InSettings.BandWidthStretch);
+			BandSettings.CenterFreq = FPseudoConstantQ::GetConstantQCenterFrequency(CQTBandIndex, InSettings.KernelLowestCenterFreq, InSettings.NumBandsPerOctave);
+			BandSettings.BandWidth = FPseudoConstantQ::GetConstantQBandWidth(BandSettings.CenterFreq, InSettings.NumBandsPerOctave, InSettings.BandWidthStretch);
 
-			if ((BandCenter - BandWidth) > InSampleRate)
+			if ((BandSettings.CenterFreq - BandSettings.BandWidth) > InSampleRate)
 			{
 				continue;
 			}
 
-			if (BandCenter > (2.f * InSampleRate))
+			if (BandSettings.CenterFreq > (2.f * InSampleRate))
 			{
 				continue;
 			}
 
 			// Create gaussian centered around center freq and with appropriate bandwidth
-			TArray<float> OffsetBandWeights;
+			AlignedFloatBuffer OffsetBandWeights;
 			int32 OffsetBandWeightsIndex = 0;
-			FillArrayWithTruncatedGaussian(BandCenter, BandWidth, InFFTSize, InSampleRate, OffsetBandWeights, OffsetBandWeightsIndex);
-
-
-			// set infs and nans to zero
-			for (int32 i = 0; i < OffsetBandWeights.Num(); i++)
-			{
-				if (!FMath::IsFinite(OffsetBandWeights[i]))
-				{
-					OffsetBandWeights[i] = 0.f;
-				}
-			}
-
-			// Need to have a lower bound for 
-			float DigitalBandWidth = FMath::Max(1.f / static_cast<float>(InFFTSize), BandWidth / InSampleRate);
-
-			// Sanity check the CQT bins to make sure we the bandwidth wasn't so small that the array is essentially empty
-			if (OffsetBandWeights.Num() > 0)
-			{
 
-				float GaussianSum = 0.f;
-				Audio::ArraySum(OffsetBandWeights, GaussianSum);
-
-				if (GaussianSum < 0.95f)
-				{
-					// It's a bit of an arbitrary threshold, but it tells us that the bandwidth
-					// is low enough and the FFT granularity course enough to where our pseudo 
-					// cqt windows will likely miss data. In this case we force the window to 
-					// interpolate between two nearest fft bins.
-					FMemory::Memset(OffsetBandWeights.GetData(), 0, OffsetBandWeights.Num() * sizeof(float));
-					float Position = InFFTSize * BandCenter / InSampleRate - static_cast<float>(OffsetBandWeightsIndex);
-					int32 LowerIndex = FMath::FloorToInt(Position);
-					int32 UpperIndex = LowerIndex + 1;
-
-					if ((LowerIndex >= 0) && (UpperIndex < OffsetBandWeights.Num()))
-					{
-						OffsetBandWeights[LowerIndex] = 1.f - (Position - static_cast<float>(LowerIndex));
-						OffsetBandWeights[UpperIndex] = 1.f - OffsetBandWeights[LowerIndex];
-
-						// Need to update digital bandwidth to be bandwidth of FFTbin
-						DigitalBandWidth = 1.f / static_cast<float>(InFFTSize);
-					}
-					else if ((LowerIndex >= 0) && (LowerIndex < OffsetBandWeights.Num()))
-					{
-						OffsetBandWeights[LowerIndex] = 1.f;
-						DigitalBandWidth = 1.f / static_cast<float>(InFFTSize);
-					}
-					else if ((UpperIndex >= 0) && (UpperIndex < OffsetBandWeights.Num()))
-					{
-						OffsetBandWeights[UpperIndex] = 1.f;
-						DigitalBandWidth = 1.f / static_cast<float>(InFFTSize);
-					}
-				}
-			}
-
-
-
-
-			// Normalize window
-			float NormDenom = 1.f;
-			switch (InSettings.Normalization)
-			{
-				case EPseudoConstantQNormalization::EqualAmplitude:
-					NormDenom = 1.f;
-					break;
-				case EPseudoConstantQNormalization::EqualEuclideanNorm:
-					NormDenom = FMath::Sqrt(DigitalBandWidth * InFFTSize * FMath::Sqrt(PI  * 2.f));
-					break;
-				case EPseudoConstantQNormalization::EqualEnergy:
-				default:
-					NormDenom = DigitalBandWidth * InFFTSize * FMath::Sqrt(PI  * 2.f);
-					break;
-			}
-
-			if ((NormDenom > 0.f) && (NormDenom != 1.f))
-			{
-				ArrayMultiplyByConstantInPlace(OffsetBandWeights, 1.f / NormDenom);
-			}
-
-			// Truncate if necessary
-			if (OffsetBandWeightsIndex >= NumUsefulFFTBins)
-			{
-				OffsetBandWeightsIndex = 0;
-				OffsetBandWeights.Empty();
-			}
-			else if ((OffsetBandWeightsIndex + OffsetBandWeights.Num()) > NumUsefulFFTBins)
-			{
-				OffsetBandWeights.SetNum(NumUsefulFFTBins - OffsetBandWeightsIndex);
-			}
+			FPseudoConstantQ::FillArrayWithConstantQBand(BandSettings, OffsetBandWeights, OffsetBandWeightsIndex);
 
 			// Store row in transform
 			Transform->SetRow(CQTBandIndex, OffsetBandWeightsIndex, OffsetBandWeights);
diff --git a/Engine/Source/Runtime/SignalProcessing/Private/SpectrumAnalyzer.cpp b/Engine/Source/Runtime/SignalProcessing/Private/SpectrumAnalyzer.cpp
--- a/Engine/Source/Runtime/SignalProcessing/Private/SpectrumAnalyzer.cpp
+++ b/Engine/Source/Runtime/SignalProcessing/Private/SpectrumAnalyzer.cpp
@@ -1,10 +1,567 @@
 // Copyright Epic Games, Inc. All Rights Reserved.
 
 #include "DSP/SpectrumAnalyzer.h"
-
+#include "DSP/FFTAlgorithm.h"
+#include "DSP/ConstantQ.h"
+#include "SignalProcessingModule.h"
 
 namespace Audio
 {
+	namespace SpectrumAnalyzerIntrinsics
+	{
+		// Bit mask for returning even numbers of int32
+		const int32 EvenNumberMask = 0xFFFFFFFE;
+
+		// Constant useful for calculating log10
+		const float Loge10 = FMath::Loge(10.f);
+	}
+
+
+	// Implementation of spectrum band extractor
+	class FSpectrumBandExtractor : public ISpectrumBandExtractor
+	{
+			// FBandSpec describes specifications for a single band.
+			struct FBandSpec
+			{
+				// Location in output array where band value should be stored.
+				int32 OutIndex;
+
+				// Center frequency of the band.
+				float CenterFrequency;
+
+				// The metric used for the band value.
+				EMetric Metric;
+
+				// The noisefloor in decibels, used when the metric is decibels. 
+				float DbNoiseFloor;
+
+				// The scaling parameter to apply to the power spectrum.
+				float PowerSpectrumScale;
+
+				// If true, all values are scaled and clamped between 0.0 and 1.0.
+				bool bDoNormalize;
+
+				FBandSpec(const FSpectrumBandExtractorSettings& InSettings, int32 InOutIndex, float InCenterFrequency, EMetric InMetric, float InDecibelNoiseFloor, bool bInDoNormalize)
+				:	OutIndex(InOutIndex)
+				,	CenterFrequency(InCenterFrequency)
+				,	Metric(InMetric)
+				,	DbNoiseFloor(InDecibelNoiseFloor)
+				,	PowerSpectrumScale(1.f)
+				,	bDoNormalize(bInDoNormalize)
+				{
+					Update(InSettings);
+				}
+
+				virtual ~FBandSpec() {}
+
+				// Update calculates parameters that are specific to FFT implementation
+				// and sample rate.
+				virtual void Update(const FSpectrumBandExtractorSettings& InSettings)
+				{
+					PowerSpectrumScale = 1.f;
+					float FloatFFTSize = FMath::Max(static_cast<float>(InSettings.FFTSize), 1.f);
+
+					switch (InSettings.FFTScaling)
+					{
+						case EFFTScaling::MultipliedByFFTSize:
+							PowerSpectrumScale = 1.f / (FloatFFTSize * FloatFFTSize);
+							break;
+
+						case EFFTScaling::MultipliedBySqrtFFTSize:
+							PowerSpectrumScale = 1.f / FloatFFTSize;
+							break;
+
+						case EFFTScaling::DividedByFFTSize:
+							PowerSpectrumScale = FloatFFTSize * FloatFFTSize;
+							break;
+
+						case EFFTScaling::DividedBySqrtFFTSize:
+							PowerSpectrumScale = FloatFFTSize;
+							break;
+
+						case EFFTScaling::None:
+						default:
+							PowerSpectrumScale = 1.f;
+							break;
+					}
+
+					float WindowScale = 1.f;
+
+
+					switch (InSettings.WindowType)
+					{
+						case EWindowType::None:
+							WindowScale = 1.f / FloatFFTSize;
+							break;
+
+						case EWindowType::Hamming:
+							// General form of scaling for generalized cosine on powers spectrum is
+							// 1 / ((1.5 * alpha^2 - alpha + 0.5) * FFTSize)
+
+							// For hamming, alpha = 0.54
+							WindowScale = 1.f / (0.3974 * FloatFFTSize);
+							break;
+
+						case EWindowType::Hann:
+							// General form of scaling for generalized cosine on powers spectrum is
+							// 1 / ((1.5 * alpha^2 - alpha + 0.5) * FFTSize)
+
+							// For hann, alpha = 0.5
+							WindowScale = 1.f / (0.375f * FloatFFTSize);
+							break;
+
+						case EWindowType::Blackman:
+							// General form of scaling for blackman windows is
+							// 1 / ((alpha_0^2 + alpha_1^2 / 2 + alpha_2^2 / 2) * FFTSize
+							//
+							// For this window alpha_0 = 0.42, alpha_1 = 0.5 and alph_2 = 0.08
+							WindowScale = 1.f / (0.3046f * FloatFFTSize);
+							break;
+
+						default:
+							// Should not get here. Means that a window type is not covered.
+							check(false);
+							WindowScale = 1.f / FloatFFTSize;
+							break;
+
+					}
+
+					PowerSpectrumScale *= WindowScale;
+				}
+			};
+
+			// Specification for a nearest neighbor band.
+			struct FNNBandSpec : public FBandSpec
+			{
+				// Use parent constructor
+				using FBandSpec::FBandSpec;
+
+				// Index in power spectrum to lookup band.
+				int32 Index;
+
+				virtual void Update(const FSpectrumBandExtractorSettings& InSettings) override
+				{
+					// Call parent class update.
+					FBandSpec::Update(InSettings);
+
+					// Update the index
+					const int32 MaxSpectrumIndex = InSettings.FFTSize / 2;
+					const float Position = CenterFrequency / FMath::Max(InSettings.SampleRate, 1.f) * InSettings.FFTSize;
+					Index = FMath::RoundToInt(Position);
+					Index = FMath::Clamp(Index, 0, MaxSpectrumIndex);
+				}
+			};
+
+			// Specification for a linearly interpolated band.
+			struct FLerpBandSpec : public FBandSpec
+			{
+				// Use parent constructor
+				using FBandSpec::FBandSpec;
+
+				// Lower index power spectrum.
+				int32 LowerIndex;
+
+				// Upper index of power spectrum.
+				int32 UpperIndex;
+
+				// Value used for lerping between lower and upper band values.
+				float Alpha;
+
+				virtual void Update(const FSpectrumBandExtractorSettings& InSettings) override
+				{
+					// Call parent class update.
+					FBandSpec::Update(InSettings);
+
+					// Update lower index, upper index and alpha.
+					const int32 MaxSpectrumIndex = InSettings.FFTSize / 2;
+					const float Position = CenterFrequency / FMath::Max(InSettings.SampleRate, 1.f) * InSettings.FFTSize;
+
+					LowerIndex = FMath::FloorToInt(Position);
+					UpperIndex = LowerIndex + 1;
+					Alpha = Position - LowerIndex;
+
+					LowerIndex = FMath::Clamp(LowerIndex, 0, MaxSpectrumIndex);
+					UpperIndex = FMath::Clamp(UpperIndex, 0, MaxSpectrumIndex);
+					Alpha = FMath::Clamp(Alpha, 0.f, 1.f);
+				}
+			};
+
+			// Specification for band using quadratic interpolation.
+			struct FQuadraticBandSpec : public FBandSpec
+			{
+				// Use parent constructor
+				using FBandSpec::FBandSpec;
+
+				// Lower index of power spectrum used for interpolation.
+				int32 LowerIndex;
+
+				// Middle index of power spectrum used for interpolation.
+				int32 MidIndex;
+
+				// Upper index of power spectrum used for interpolation.
+				int32 UpperIndex;
+
+				// Weight for lower value.
+				float LowerWeight;
+
+				// Weight for middle value.
+				float MidWeight;
+
+				// Weight for upper value.
+				float UpperWeight;
+
+				virtual void Update(const FSpectrumBandExtractorSettings& InSettings) override
+				{
+					// Call parent class update.
+					FBandSpec::Update(InSettings);
+
+					// Update indices and weights.
+					const int32 MaxSpectrumIndex = InSettings.FFTSize / 2;
+					const float Position = CenterFrequency / FMath::Max(InSettings.SampleRate, 1.f) * InSettings.FFTSize;
+
+					MidIndex = FMath::RoundToInt(Position);
+					LowerIndex = MidIndex - 1;
+					UpperIndex = MidIndex + 1;
+
+					// Calculate polynomail weights 
+					float RelativePosition = Position - LowerIndex;
+
+					LowerWeight = ((RelativePosition - 1.f) * (RelativePosition - 2.f)) / 2.f;
+					MidWeight = (RelativePosition * (RelativePosition - 2.f)) / -1.f;
+					UpperWeight = (RelativePosition * (RelativePosition - 1.f)) / 2.f;
+
+					LowerIndex = FMath::Clamp(LowerIndex, 0, MaxSpectrumIndex);
+					MidIndex = FMath::Clamp(MidIndex, 0, MaxSpectrumIndex);
+					UpperIndex = FMath::Clamp(UpperIndex, 0, MaxSpectrumIndex);
+				}
+			};
+
+			// Specification for band using CQT band.
+			struct FCQTBandSpec : public FBandSpec
+			{
+				// Use parent constructor
+				using FBandSpec::FBandSpec;
+
+				// QFactor controls the band width.
+				float QFactor;
+				
+				// Start index in power spectrum 
+				int32 StartIndex;
+
+				// Weights (offset by start index) to apply to power spectrum
+				AlignedFloatBuffer Weights;
+
+				// Internal buffer used when calculating band.
+				mutable AlignedFloatBuffer WorkBuffer;
+
+				virtual void Update(const FSpectrumBandExtractorSettings& InSettings) override
+				{
+					// Call parent class update.
+					FBandSpec::Update(InSettings);
+
+					// Update band weights and offset index.
+					const int32 MaxSpectrumIndex = InSettings.FFTSize / 2;
+					const float Position = CenterFrequency / FMath::Max(InSettings.SampleRate, 1.f) * InSettings.FFTSize;
+					FPseudoConstantQBandSettings CQTBandSettings;
+					CQTBandSettings.CenterFreq = CenterFrequency;
+					CQTBandSettings.BandWidth = FMath::Max(SMALL_NUMBER, CenterFrequency / FMath::Max(SMALL_NUMBER, QFactor));
+					CQTBandSettings.FFTSize = InSettings.FFTSize;
+					CQTBandSettings.SampleRate = FMath::Max(1.f, InSettings.SampleRate);
+					CQTBandSettings.Normalization = EPseudoConstantQNormalization::EqualEnergy;
+
+					StartIndex = 0;
+					Weights.Reset();
+					WorkBuffer.Reset();
+
+					FPseudoConstantQ::FillArrayWithConstantQBand(CQTBandSettings, Weights, StartIndex);
+
+					if (Weights.Num() > 0)
+					{
+						WorkBuffer.AddUninitialized(Weights.Num());
+					}
+				}
+			};
+
+
+		public:
+			FSpectrumBandExtractor(const FSpectrumBandExtractorSettings& InSettings)
+			:	Settings(InSettings)
+			{
+			}
+
+			virtual void SetSettings(const FSpectrumBandExtractorSettings& InSettings) override
+			{
+				bool bSettingsChanged = (Settings != InSettings);
+				Settings = InSettings;
+
+				if (bSettingsChanged)
+				{
+
+					// If the settings have changed from the previous call, the band specs
+					// need to be updated with the new information.
+					UpdateBandSpecs();
+				}
+			}
+			
+			// Clear out all added bands.
+			virtual void RemoveAllBands() override
+			{
+				NNBandSpecs.Reset();
+				LerpBandSpecs.Reset();
+				QuadraticBandSpecs.Reset();
+				CQTBandSpecs.Reset();
+			}
+
+			// Return total number of bands.
+			virtual int32 GetNumBands() const override
+			{
+				int32 Num = NNBandSpecs.Num();
+				Num += LerpBandSpecs.Num();
+				Num += QuadraticBandSpecs.Num();
+				Num += CQTBandSpecs.Num();
+
+				return Num;
+			}
+
+			// Add a nearest neighbor band
+			virtual void AddNearestNeighborBand(float InCenterFrequency, EMetric InMetric, float InDecibelNoiseFloor, bool bInDoNormalize) override
+			{
+				AddBand<FNNBandSpec>(NNBandSpecs, InCenterFrequency, InMetric, InDecibelNoiseFloor, bInDoNormalize);
+			}
+
+			// Add a linear interpolated band.
+			virtual void AddLerpBand(float InCenterFrequency, EMetric InMetric, float InDecibelNoiseFloor, bool bInDoNormalize) override
+			{
+				AddBand<FLerpBandSpec>(LerpBandSpecs, InCenterFrequency, InMetric, InDecibelNoiseFloor, bInDoNormalize);
+			}
+
+			// Add a quadratic interplated band.
+			virtual void AddQuadraticBand(float InCenterFrequency, EMetric InMetric, float InDecibelNoiseFloor, bool bInDoNormalize) override
+			{
+				AddBand<FQuadraticBandSpec>(QuadraticBandSpecs, InCenterFrequency, InMetric, InDecibelNoiseFloor, bInDoNormalize);
+			}
+
+			// Add a constant Q band
+			virtual void AddConstantQBand(float InCenterFrequency, float InQFactor, EMetric InMetric, float InDecibelNoiseFloor, bool bInDoNormalize) override
+			{
+				FCQTBandSpec& Band = AddBand<FCQTBandSpec>(CQTBandSpecs, InCenterFrequency, InMetric, InDecibelNoiseFloor, bInDoNormalize);
+				Band.QFactor = FMath::Clamp(InQFactor, 0.001f, 100.f);
+			}
+
+			// Extract band from input.
+			virtual void ExtractBands(const AlignedFloatBuffer& InComplexBuffer, TArray<float>& OutValues) override
+			{
+				const int32 NumComplex = InComplexBuffer.Num();
+
+				check(NumComplex == (Settings.FFTSize + 2));
+
+				OutValues.Reset();
+				OutValues.AddZeroed(GetNumBands());
+
+				PowerSpectrum.Reset();
+				if (NumComplex > 1)
+				{
+					PowerSpectrum.AddUninitialized(NumComplex / 2);
+				}
+
+				// All band extractors expect a power spectrum
+				ArrayComplexToPower(InComplexBuffer, PowerSpectrum);
+
+				ExtractBands(PowerSpectrum, NNBandSpecs, OutValues);
+				ExtractBands(PowerSpectrum, LerpBandSpecs, OutValues);
+				ExtractBands(PowerSpectrum, QuadraticBandSpecs, OutValues);
+				ExtractBands(PowerSpectrum, CQTBandSpecs, OutValues);
+			}
+
+		private:
+
+			// Adds a band spec and returns a reference to the added spec.
+			template<typename T> 
+			T& AddBand(TArray<T>& InBandSpecs, float InCenterFrequency, EMetric InMetric, float InDecibelNoiseFloor, bool bInDoNormalize)
+			{
+				int32 OutIndex = GetNumBands();
+
+				T BandSpec(Settings, OutIndex, InCenterFrequency, InMetric, InDecibelNoiseFloor, bInDoNormalize);
+
+				return InBandSpecs.Add_GetRef(BandSpec);
+			}
+
+			// Calls update on all band specs in the array.
+			template<typename T>
+			void UpdateBandSpecs(TArray<T>& InSpecs)
+			{
+				for (T& BandSpec : InSpecs)
+				{
+					BandSpec.Update(Settings);
+				}
+			}
+
+			// Updates all band specs.
+			void UpdateBandSpecs()
+			{
+				UpdateBandSpecs<FNNBandSpec>(NNBandSpecs);
+				UpdateBandSpecs<FLerpBandSpec>(LerpBandSpecs);
+				UpdateBandSpecs<FQuadraticBandSpec>(QuadraticBandSpecs);
+				UpdateBandSpecs<FCQTBandSpec>(CQTBandSpecs);
+			}
+
+
+			// Extract nearest neighbor bands
+			void ExtractBands(const AlignedFloatBuffer& InPowerSpectrum, const TArray<FNNBandSpec>& InNNBandSpecs, TArray<float>& OutValues) const
+			{
+				float* OutData = OutValues.GetData();
+				const float* InData = InPowerSpectrum.GetData();
+				int32 InNum = InPowerSpectrum.Num();
+
+				for (const FNNBandSpec& Spec : InNNBandSpecs)
+				{
+					check(Spec.OutIndex >= 0);
+					check(Spec.OutIndex < OutValues.Num());
+					check(Spec.Index < InNum);
+					check(Spec.Index >= 0);
+					OutData[Spec.OutIndex] = ApplyScaleAndMetric(Spec, InData[Spec.Index]);
+				};
+			}
+
+			// Extract linearly interpolated bands
+			void ExtractBands(const AlignedFloatBuffer& InPowerSpectrum, const TArray<FLerpBandSpec>& InLerpBandSpecs, TArray<float>& OutValues) const
+			{
+				float* OutData = OutValues.GetData();
+				const float* InData = InPowerSpectrum.GetData();
+				int32 InNum = InPowerSpectrum.Num();
+
+				for (const FLerpBandSpec& Spec : InLerpBandSpecs)
+				{
+					check(Spec.OutIndex >= 0);
+					check(Spec.OutIndex < OutValues.Num());
+					check(Spec.LowerIndex < InNum);
+					check(Spec.LowerIndex >= 0);
+					check(Spec.UpperIndex < InNum);
+					check(Spec.UpperIndex >= 0);
+
+					float Value = FMath::Lerp<float>(InData[Spec.LowerIndex], InData[Spec.UpperIndex], Spec.Alpha);
+					OutData[Spec.OutIndex] = ApplyScaleAndMetric(Spec, Value);
+				};
+			}
+
+			// Extract quadratically interpolated bands.
+			void ExtractBands(const AlignedFloatBuffer& InPowerSpectrum, const TArray<FQuadraticBandSpec>& InQuadraticBandSpecs, TArray<float>& OutValues) const
+			{
+				float* OutData = OutValues.GetData();
+				const float* InData = InPowerSpectrum.GetData();
+				int32 InNum = InPowerSpectrum.Num();
+
+				for (const FQuadraticBandSpec& Spec : InQuadraticBandSpecs)
+				{
+					check(Spec.OutIndex >= 0);
+					check(Spec.OutIndex < OutValues.Num());
+					check(Spec.LowerIndex < InNum);
+					check(Spec.LowerIndex >= 0);
+					check(Spec.MidIndex < InNum);
+					check(Spec.MidIndex >= 0);
+					check(Spec.UpperIndex < InNum);
+					check(Spec.UpperIndex >= 0);
+
+					const float LowerValue = InData[Spec.LowerIndex];
+					const float MidValue = InData[Spec.MidIndex];
+					const float UpperValue = InData[Spec.UpperIndex];
+					
+					float Value = (LowerValue * Spec.LowerWeight) + (MidValue * Spec.MidWeight) + (UpperValue * Spec.UpperWeight);
+
+					OutData[Spec.OutIndex] = ApplyScaleAndMetric(Spec, Value);
+				}
+			}
+
+			// Extract constant q bands.
+			void ExtractBands(const AlignedFloatBuffer& InPowerSpectrum, const TArray<FCQTBandSpec>& InCQTBandSpecs, TArray<float>& OutValues) const
+			{
+				float* OutData = OutValues.GetData();
+				const float* InData = InPowerSpectrum.GetData();
+				int32 InNum = InPowerSpectrum.Num();
+
+				for (const FCQTBandSpec& Spec : InCQTBandSpecs)
+				{
+					check(Spec.OutIndex >= 0);
+					check(Spec.OutIndex < OutValues.Num());
+					check(Spec.StartIndex < InNum);
+					check(Spec.StartIndex >= 0);
+					check((Spec.StartIndex + Spec.Weights.Num()) <= InNum);
+
+					float Value = 0.f;
+					int32 NumWeights = Spec.Weights.Num();
+
+					if (NumWeights > 0)
+					{
+						check(NumWeights == Spec.WorkBuffer.Num());
+						FMemory::Memcpy(Spec.WorkBuffer.GetData(), &InData[Spec.StartIndex], NumWeights * sizeof(float));
+
+						ArrayMultiplyInPlace(Spec.Weights, Spec.WorkBuffer);
+						ArraySum(Spec.WorkBuffer, Value);
+					}
+
+					OutData[Spec.OutIndex] = ApplyScaleAndMetric(Spec, Value);
+				}
+				
+			}
+
+			// Scale and apply metric to a band value
+			float ApplyScaleAndMetric(const FBandSpec& InBandSpec, float InValue) const
+			{
+				float OutValue = InValue * InBandSpec.PowerSpectrumScale;
+
+				switch (InBandSpec.Metric)
+				{
+					case ISpectrumBandExtractor::EMetric::Magnitude:
+						OutValue = FMath::Sqrt(OutValue);
+						break;
+
+					case ISpectrumBandExtractor::EMetric::Decibel:
+						OutValue = 10.f * FMath::Loge(OutValue) / SpectrumAnalyzerIntrinsics::Loge10;
+						if (!FMath::IsFinite(OutValue) || (OutValue < InBandSpec.DbNoiseFloor))
+						{
+							OutValue = InBandSpec.DbNoiseFloor;
+						}
+
+						if (InBandSpec.bDoNormalize)
+						{
+							OutValue -= InBandSpec.DbNoiseFloor;
+							if (InBandSpec.DbNoiseFloor < 0.f)
+							{
+								OutValue /= (-InBandSpec.DbNoiseFloor);
+							}
+						}
+						break;
+
+					case ISpectrumBandExtractor::EMetric::Power:
+					default:
+						OutValue = InValue;
+						break;
+				}
+
+				if (InBandSpec.bDoNormalize)
+				{
+					OutValue = FMath::Clamp(OutValue, 0.f, 1.f);
+				}
+
+				return OutValue;
+			}
+
+			FSpectrumBandExtractorSettings Settings;
+
+			AlignedFloatBuffer PowerSpectrum;
+			TArray<FNNBandSpec> NNBandSpecs;
+			TArray<FLerpBandSpec> LerpBandSpecs;
+			TArray<FQuadraticBandSpec> QuadraticBandSpecs;
+			TArray<FCQTBandSpec> CQTBandSpecs;
+	};
+
+
+	// Creates a concreted implementation of teh ISpectrumBandExtractor interface.
+	TUniquePtr<ISpectrumBandExtractor> ISpectrumBandExtractor::CreateSpectrumBandExtractor(const FSpectrumBandExtractorSettings& InSettings)
+	{
+		return MakeUnique<FSpectrumBandExtractor>(InSettings);
+	}
+
 	FSpectrumAnalyzer::FSpectrumAnalyzer()
 		: CurrentSettings(FSpectrumAnalyzerSettings())
 		, bSettingsWereUpdated(false)
@@ -12,7 +569,7 @@ namespace Audio
 		, SampleRate(0.0f)
 		, Window(CurrentSettings.WindowType, (int32)CurrentSettings.FFTSize, 1, false)
 		, InputQueue(FMath::Max((int32)CurrentSettings.FFTSize * 4, 4096))
-		, FrequencyBuffer(CurrentSettings)
+		, FrequencyBuffer((int32)CurrentSettings.FFTSize)
 		, LockedFrequencyVector(nullptr)
 	{
 	}
@@ -24,7 +581,7 @@ namespace Audio
 		, SampleRate(InSampleRate)
 		, Window(InSettings.WindowType, (int32)InSettings.FFTSize, 1, false)
 		, InputQueue(FMath::Max((int32)CurrentSettings.FFTSize * 4, 4096))
-		, FrequencyBuffer(InSettings)
+		, FrequencyBuffer((int32)InSettings.FFTSize)
 		, LockedFrequencyVector(nullptr)
 	{
 		ResetSettings();
@@ -37,7 +594,7 @@ namespace Audio
 		, SampleRate(InSampleRate)
 		, Window(CurrentSettings.WindowType, (int32)CurrentSettings.FFTSize, 1, false)
 		, InputQueue(FMath::Max((int32)CurrentSettings.FFTSize * 4, 4096))
-		, FrequencyBuffer(CurrentSettings)
+		, FrequencyBuffer((int32)CurrentSettings.FFTSize)
 		, LockedFrequencyVector(nullptr)
 	{
 		ResetSettings();
@@ -63,7 +620,7 @@ namespace Audio
 		bSettingsWereUpdated = false;
 		SampleRate = InSampleRate;
 		InputQueue.SetCapacity(FMath::Max((int32)CurrentSettings.FFTSize * 4, 4096));
-		FrequencyBuffer.Reset(CurrentSettings);
+		FrequencyBuffer.Reset((int32)CurrentSettings.FFTSize);
 		ResetSettings();
 
 		bIsInitialized = true;
@@ -80,6 +637,15 @@ namespace Audio
 
 		Window = FWindow(CurrentSettings.WindowType, (int32)CurrentSettings.FFTSize, 1, false);
 		FFTSize = (int32) CurrentSettings.FFTSize;
+		int32 Log2FFTSize = 9;
+		if (FFTSize > 0)
+		{
+			// FFTSize must be log2
+			check(FMath::CountBits(FFTSize) == 1);
+			Log2FFTSize = FMath::CountTrailingZeros(FFTSize);
+		}
+
+		AnalysisTimeDomainBuffer.Reset();
 		
 		if (FMath::IsNearlyZero(CurrentSettings.HopSize))
 		{
@@ -89,73 +655,117 @@ namespace Audio
 		{
 			HopInSamples = FMath::FloorToInt((float)CurrentSettings.FFTSize * CurrentSettings.HopSize);
 		}
-		
-		AnalysisTimeDomainBuffer.Reset();
-		AnalysisTimeDomainBuffer.AddZeroed(FFTSize);
 
-		FrequencyBuffer.Reset(CurrentSettings);
+		// Create a new FFT
+		FFFTSettings FFTSettings;
+		FFTSettings.Log2Size = Log2FFTSize;
+		FFTSettings.bArrays128BitAligned = true;
+		FFTSettings.bEnableHardwareAcceleration = true;
+
+		FFT = FFFTFactory::NewFFTAlgorithm(FFTSettings);
+
+		if (!FFT.IsValid())
+		{
+			if (FFFTFactory::AreFFTSettingsSupported(FFTSettings))
+			{
+				UE_LOG(LogSignalProcessing, Error, TEXT("Failed to create fft for supported settings."))
+			}
+			else
+			{
+				UE_LOG(LogSignalProcessing, Warning, TEXT("FFT Settings are unsupported."))
+			}
+			FFTScaling = EFFTScaling::None;
+
+			if (FFTSize > 0)
+			{
+				AnalysisTimeDomainBuffer.AddZeroed(FFTSize);
+				FrequencyBuffer.Reset(FFTSize);
+			}
+		}
+		else
+		{
+			int32 NumFFTInput = FFT->NumInputFloats();
+			int32 NumFFTOutput = FFT->NumOutputFloats();
+			FFTScaling = FFT->ForwardScaling();
+
+			if (NumFFTInput > 0)
+			{
+				AnalysisTimeDomainBuffer.AddUninitialized(NumFFTInput);
+			}
+
+			FrequencyBuffer.Reset(NumFFTOutput);
+		}
+		
 		bSettingsWereUpdated = false;
 	}
 
-	void FSpectrumAnalyzer::PerformInterpolation(const FSpectrumAnalyzerFrequencyVector* InFrequencies, FSpectrumAnalyzerSettings::EPeakInterpolationMethod InMethod, const float InFreq, float& OutReal, float& OutImag)
+	void FSpectrumAnalyzer::PerformInterpolation(const AlignedFloatBuffer& InComplexBuffer, FSpectrumAnalyzerSettings::EPeakInterpolationMethod InMethod, const float InFreq, float& OutReal, float& OutImag)
 	{
-		const int32 VectorLength = InFrequencies->RealFrequencies.Num();
-		const float NyquistPosition = VectorLength / 2;
+		const float* InComplexData = InComplexBuffer.GetData();
+		const int32 VectorLength = InComplexBuffer.Num();
+		const int32 NyquistPosition = VectorLength - 2;
 		
-		const float Nyquist = SampleRate / 2;
+		const float Nyquist = SampleRate / 2.f;
 
 		// Fractional position in the frequency vector in terms of indices.
 		// float Position = NyquistPosition + (InFreq / Nyquist);
 		const float NormalizedFreq = (InFreq / Nyquist);
-		float Position = InFreq >= 0 ? (NormalizedFreq * VectorLength / 2) : VectorLength - (NormalizedFreq * VectorLength / 2);
+		float Position = InFreq >= 0 ? (NormalizedFreq * VectorLength) : 0.f;
 		
-		// Position should be clamped between just above DC and just below nyquist to avoid rounding errors.
-		Position = FMath::Clamp<float>(Position, 0.01f, VectorLength - 1.01f);
-
 		switch (InMethod)
 		{
 		case Audio::FSpectrumAnalyzerSettings::EPeakInterpolationMethod::NearestNeighbor:
 		{
-			int32 Index = FMath::RoundToInt(Position);
-			OutReal = InFrequencies->RealFrequencies[Index];
-			OutImag = InFrequencies->ImagFrequencies[Index];
+			int32 Index = FMath::RoundToInt(Position) & SpectrumAnalyzerIntrinsics::EvenNumberMask;
+
+			Index = FMath::Clamp(Index, 0, NyquistPosition);
+
+			OutReal = InComplexData[Index];
+			OutImag = InComplexData[Index + 1];
+
 			break;
 		}
 			
 		case Audio::FSpectrumAnalyzerSettings::EPeakInterpolationMethod::Linear:
 		{
-			const int32 LowerIndex = FMath::FloorToInt(Position);
-			const int32 UpperIndex = FMath::CeilToInt(Position);
+			int32 LowerIndex = FMath::FloorToInt(Position) & SpectrumAnalyzerIntrinsics::EvenNumberMask;
+			int32 UpperIndex = LowerIndex + 2;
+
+			LowerIndex = FMath::Clamp(LowerIndex, 0, NyquistPosition);
+			UpperIndex = FMath::Clamp(UpperIndex, 0, NyquistPosition);
+
 			const float PositionFraction = Position - LowerIndex;
 
-			const float y1Real = InFrequencies->RealFrequencies[LowerIndex];
-			const float y2Real = InFrequencies->RealFrequencies[UpperIndex];
+			const float y1Real = InComplexData[LowerIndex];
+			const float y2Real = InComplexData[UpperIndex];
+
 			OutReal = FMath::Lerp<float>(y1Real, y1Real, PositionFraction);
 
-			const float y1Imag = InFrequencies->ImagFrequencies[LowerIndex];
-			const float y2Imag = InFrequencies->ImagFrequencies[UpperIndex];
+			const float y1Imag = InComplexData[LowerIndex + 1];
+			const float y2Imag = InComplexData[UpperIndex + 1];
+
 			OutImag = FMath::Lerp<float>(y1Imag, y2Imag, PositionFraction);
 			break;
 		}	
 		case Audio::FSpectrumAnalyzerSettings::EPeakInterpolationMethod::Quadratic:
 		{
-			const int32 MidIndex = FMath::RoundToInt(Position);
-			const int32 LowerIndex = FMath::Max(0, MidIndex - 1);
-			const int32 UpperIndex = FMath::Min(VectorLength, MidIndex + 1);
-			const float PositionFraction = Position - LowerIndex;
+			// Note: math here does not interpolate quadratically. 
+			const int32 MidIndex = FMath::Clamp(FMath::RoundToInt(Position) & SpectrumAnalyzerIntrinsics::EvenNumberMask, 0, NyquistPosition);
+			const int32 LowerIndex = FMath::Max(0, MidIndex - 2);
+			const int32 UpperIndex = FMath::Min(NyquistPosition, MidIndex + 2);
 
-			const float y1Real = InFrequencies->RealFrequencies[LowerIndex];
-			const float y2Real = InFrequencies->RealFrequencies[MidIndex];
-			const float y3Real = InFrequencies->RealFrequencies[UpperIndex];
+			const float y1Real = InComplexData[LowerIndex];
+			const float y2Real = InComplexData[MidIndex];
+			const float y3Real = InComplexData[UpperIndex];
 
-			const float InterpReal = (y3Real - y1Real) / (2 * (2 * y2Real - y1Real - y3Real));
+			const float InterpReal = (y3Real - y1Real) / (2.f * (2.f * y2Real - y1Real - y3Real));
 			
 			OutReal = InterpReal;
 
-			const float y1Imag = InFrequencies->ImagFrequencies[LowerIndex];
-			const float y2Imag = InFrequencies->ImagFrequencies[MidIndex];
-			const float y3Imag = InFrequencies->ImagFrequencies[UpperIndex];
-			const float InterpImag = (y3Imag - y1Imag) / (2 * (2 * y2Imag - y1Imag - y3Imag));
+			const float y1Imag = InComplexData[LowerIndex + 1];
+			const float y2Imag = InComplexData[MidIndex + 1];
+			const float y3Imag = InComplexData[UpperIndex + 1];
+			const float InterpImag = (y3Imag - y1Imag) / (2.f * (2.f * y2Imag - y1Imag - y3Imag));
 
 			OutImag = InterpImag;
 			break;
@@ -181,10 +791,10 @@ namespace Audio
 	{
 		if (!bIsInitialized)
 		{
-			return false;
+			return 0.f;
 		}
 
-		const FSpectrumAnalyzerFrequencyVector* OutVector;
+		const AlignedFloatBuffer* OutVector = nullptr;
 		bool bShouldUnlockBuffer = true;
 
 		if (LockedFrequencyVector)
@@ -194,7 +804,7 @@ namespace Audio
 		}
 		else
 		{
-			OutVector = FrequencyBuffer.LockMostRecentBuffer();
+			OutVector = &FrequencyBuffer.LockMostRecentBuffer();
 		}
 
 		// Perform work.
@@ -203,7 +813,7 @@ namespace Audio
 			float OutMagnitude = 0.0f;
 
 			float InterpolatedReal, InterpolatedImag;
-			PerformInterpolation(OutVector, CurrentSettings.InterpolationMethod, InFrequency, InterpolatedReal, InterpolatedImag);
+			PerformInterpolation(*OutVector, CurrentSettings.InterpolationMethod, InFrequency, InterpolatedReal, InterpolatedImag);
 
 			OutMagnitude = FMath::Sqrt((InterpolatedReal * InterpolatedReal) + (InterpolatedImag * InterpolatedImag));
 
@@ -223,10 +833,10 @@ namespace Audio
 	{
 		if (!bIsInitialized)
 		{
-			return false;
+			return 0.f;
 		}
 
-		const FSpectrumAnalyzerFrequencyVector* OutVector;
+		const AlignedFloatBuffer* OutVector = nullptr;
 		bool bShouldUnlockBuffer = true;
 
 		if (LockedFrequencyVector)
@@ -236,7 +846,7 @@ namespace Audio
 		}
 		else
 		{
-			OutVector = FrequencyBuffer.LockMostRecentBuffer();
+			OutVector = &FrequencyBuffer.LockMostRecentBuffer();
 		}
 
 		// Perform work.
@@ -245,7 +855,7 @@ namespace Audio
 			float OutPhase = 0.0f;
 
 			float InterpolatedReal, InterpolatedImag;
-			PerformInterpolation(OutVector, CurrentSettings.InterpolationMethod, InFrequency, InterpolatedReal, InterpolatedImag);
+			PerformInterpolation(*OutVector, CurrentSettings.InterpolationMethod, InFrequency, InterpolatedReal, InterpolatedImag);
 
 			OutPhase = FMath::Atan2(InterpolatedImag, InterpolatedReal);
 
@@ -261,6 +871,50 @@ namespace Audio
 		return 0.0f;
 	}
 
+	// Return bands extracted by band extractor.
+	void FSpectrumAnalyzer::GetBands(ISpectrumBandExtractor& InExtractor, TArray<float>& OutValues) 
+	{
+		OutValues.Reset();
+
+		if (!bIsInitialized)
+		{
+			return;
+		}
+
+		const AlignedFloatBuffer* OutVector = nullptr;
+		bool bShouldUnlockBuffer = true;
+
+		FSpectrumBandExtractorSettings ExtractorSettings;
+		ExtractorSettings.SampleRate = SampleRate;
+		ExtractorSettings.FFTSize = FFTSize;
+		ExtractorSettings.FFTScaling = FFTScaling;
+		ExtractorSettings.WindowType = Window.GetWindowType();
+
+		// This should have minimal cost if settings have not changed between calls.
+		InExtractor.SetSettings(ExtractorSettings);
+
+		if (LockedFrequencyVector)
+		{
+			OutVector = LockedFrequencyVector;
+			bShouldUnlockBuffer = false;
+		}
+		else
+		{
+			OutVector = &FrequencyBuffer.LockMostRecentBuffer();
+		}
+
+		// Perform work.
+		if (OutVector)
+		{
+			InExtractor.ExtractBands(*OutVector, OutValues);
+
+			if (bShouldUnlockBuffer)
+			{
+				FrequencyBuffer.UnlockBuffer();
+			}
+		}
+	}
+
 	void FSpectrumAnalyzer::LockOutputBuffer()
 	{
 		if (!bIsInitialized)
@@ -273,7 +927,7 @@ namespace Audio
 			FrequencyBuffer.UnlockBuffer();
 		}
 
-		LockedFrequencyVector = FrequencyBuffer.LockMostRecentBuffer();
+		LockedFrequencyVector = &FrequencyBuffer.LockMostRecentBuffer();
 	}
 
 	void FSpectrumAnalyzer::UnlockOutputBuffer()
@@ -330,11 +984,11 @@ namespace Audio
 			ResetSettings();
 		}
 
-		FSpectrumAnalyzerFrequencyVector* OutputVector = FrequencyBuffer.StartWorkOnBuffer();
+		AlignedFloatBuffer& FFTOutput = FrequencyBuffer.StartWorkOnBuffer();
 
 		// If we have enough audio pushed to the spectrum analyzer and we have an available buffer to work in,
 		// we can start analyzing.
-		if (InputQueue.Num() >= ((uint32)FFTSize) && OutputVector)
+		if (InputQueue.Num() >= ((uint32)FFTSize))
 		{
 			float* TimeDomainBuffer = AnalysisTimeDomainBuffer.GetData();
 
@@ -353,15 +1007,20 @@ namespace Audio
 			Window.ApplyToBuffer(TimeDomainBuffer);
 
 			// Perform FFT.
-			FFTTimeDomainData InputParams;
-			InputParams.Buffer = TimeDomainBuffer;
-			InputParams.NumSamples = FFTSize;
-
-			FFTFreqDomainData OutputParams;
-			OutputParams.OutReal = OutputVector->RealFrequencies.GetData();
-			OutputParams.OutImag = OutputVector->ImagFrequencies.GetData();
+			if (FFT.IsValid())
+			{
+				check(AnalysisTimeDomainBuffer.Num() == FFT->NumInputFloats());
+				check(FFTOutput.Num() == FFT->NumOutputFloats());
 
-			PerformFFT(InputParams, OutputParams);
+				FFT->ForwardRealToComplex(TimeDomainBuffer, FFTOutput.GetData());
+			}
+			else
+			{
+				if (FFTOutput.Num() > 0)
+				{
+					FMemory::Memset(FFTOutput.GetData(), 0, sizeof(float) * FFTOutput.Num());
+				}
+			}
 
 			// We're done, so unlock this vector.
 			FrequencyBuffer.StopWorkOnBuffer();
@@ -387,25 +1046,30 @@ namespace Audio
 	{
 	}
 
-	FSpectrumAnalyzerBuffer::FSpectrumAnalyzerBuffer(const FSpectrumAnalyzerSettings& InSettings)
+	FSpectrumAnalyzerBuffer::FSpectrumAnalyzerBuffer(int32 InNum)
 	{
-		Reset(InSettings);
+		Reset(InNum);
 	}
 
-	void FSpectrumAnalyzerBuffer::Reset(const FSpectrumAnalyzerSettings& InSettings)
+	void FSpectrumAnalyzerBuffer::Reset(int32 InNum)
 	{
 		FScopeLock ScopeLock(&BufferIndicesCriticalSection);
 
 		static_assert(SpectrumAnalyzerBufferSize > 2, "Please ensure that SpectrumAnalyzerBufferSize is greater than 2.");
 		
-		FrequencyVectors.Reset();
+		ComplexBuffers.Reset();
 
 		for (int32 Index = 0; Index < SpectrumAnalyzerBufferSize; Index++)
 		{
-			FrequencyVectors.Emplace((int32)InSettings.FFTSize);
+			AlignedFloatBuffer& Buffer = ComplexBuffers.Emplace_GetRef();
+
+			if (InNum > 0)
+			{
+				Buffer.AddZeroed(InNum);
+			}
 		}
 
-		InputIndex = 1;
+		InputIndex = 0;
 		OutputIndex = 0;
 	}
 
@@ -435,9 +1099,9 @@ namespace Audio
 		check(InputIndex != OutputIndex);
 	}
 
-	FSpectrumAnalyzerFrequencyVector* FSpectrumAnalyzerBuffer::StartWorkOnBuffer()
+	AlignedFloatBuffer& FSpectrumAnalyzerBuffer::StartWorkOnBuffer()
 	{
-		return &(FrequencyVectors[InputIndex]);
+		return ComplexBuffers[InputIndex];
 	}
 
 	void FSpectrumAnalyzerBuffer::StopWorkOnBuffer()
@@ -445,9 +1109,9 @@ namespace Audio
 		IncrementInputIndex();
 	}
 
-	const FSpectrumAnalyzerFrequencyVector* FSpectrumAnalyzerBuffer::LockMostRecentBuffer()
+	const AlignedFloatBuffer& FSpectrumAnalyzerBuffer::LockMostRecentBuffer() const
 	{
-		return &(FrequencyVectors[OutputIndex]);
+		return ComplexBuffers[OutputIndex];
 	}
 
 	void FSpectrumAnalyzerBuffer::UnlockBuffer()
@@ -455,18 +1119,10 @@ namespace Audio
 		IncrementOutputIndex();
 	}
 
-	FSpectrumAnalyzerFrequencyVector::FSpectrumAnalyzerFrequencyVector(int32 InFFTSize)
-	{
-		RealFrequencies.Reset();
-		RealFrequencies.AddZeroed(InFFTSize);
-
-		ImagFrequencies.Reset();
-		ImagFrequencies.AddZeroed(InFFTSize);
-	}
 
 	void FSpectrumAnalysisAsyncWorker::DoWork()
 	{
 		Analyzer->PerformAnalysisIfPossible(bUseLatestAudio, false);
 	}
-
 }
+
diff --git a/Engine/Source/Runtime/SignalProcessing/Public/DSP/AudioFFT.h b/Engine/Source/Runtime/SignalProcessing/Public/DSP/AudioFFT.h
--- a/Engine/Source/Runtime/SignalProcessing/Public/DSP/AudioFFT.h
+++ b/Engine/Source/Runtime/SignalProcessing/Public/DSP/AudioFFT.h
@@ -80,6 +80,11 @@ namespace Audio
 			MultiplyBuffersInPlace(WindowBuffer.GetData(), InBuffer, NumSamples);
 		}
 
+		EWindowType GetWindowType() const
+		{
+			return WindowType;
+		}
+
 	private:
 		EWindowType WindowType;
 		AlignedFloatBuffer WindowBuffer;
diff --git a/Engine/Source/Runtime/SignalProcessing/Public/DSP/ConstantQ.h b/Engine/Source/Runtime/SignalProcessing/Public/DSP/ConstantQ.h
--- a/Engine/Source/Runtime/SignalProcessing/Public/DSP/ConstantQ.h
+++ b/Engine/Source/Runtime/SignalProcessing/Public/DSP/ConstantQ.h
@@ -52,4 +52,35 @@ namespace Audio
 	 */
 	SIGNALPROCESSING_API TUniquePtr<FContiguousSparse2DKernelTransform> NewPseudoConstantQKernelTransform(const FPseudoConstantQKernelSettings& InSettings, const int32 InFFTSize, const float InSampleRate);
 
+	/** Settings for a single constant q band. */
+	struct SIGNALPROCESSING_API FPseudoConstantQBandSettings
+	{
+		/** Center frequency of band in Hz. */
+		float CenterFreq;
+
+		/** Width of band in Hz. */
+		float BandWidth; 
+
+		/** Size of FFT used. */
+		int32 FFTSize; 
+
+		/** Samplerate of audio. */
+		float SampleRate; 
+
+		/** Desired normalization. */
+		EPseudoConstantQNormalization Normalization;
+	};
+
+	struct SIGNALPROCESSING_API FPseudoConstantQ
+	{
+		// Returns the center frequency of a given CQT band.
+		static float GetConstantQCenterFrequency(const int32 InBandIndex, const float InBaseFrequency, const float InBandsPerOctave);
+
+		// Returns the bandwidth for a given CQT band.
+		static float GetConstantQBandWidth(const float InBandCenter, const float InBandsPerOctave, const float InBandWidthStretch);
+
+		// Fills an array with the spectrum weights to apply to an FFT magnitude or power spectrum output.
+		static void FillArrayWithConstantQBand(const FPseudoConstantQBandSettings& InSettings, AlignedFloatBuffer& OutOffsetArray, int32& OutOffsetIndex);
+
+	};
 }
diff --git a/Engine/Source/Runtime/SignalProcessing/Public/DSP/SpectrumAnalyzer.h b/Engine/Source/Runtime/SignalProcessing/Public/DSP/SpectrumAnalyzer.h
--- a/Engine/Source/Runtime/SignalProcessing/Public/DSP/SpectrumAnalyzer.h
+++ b/Engine/Source/Runtime/SignalProcessing/Public/DSP/SpectrumAnalyzer.h
@@ -4,15 +4,15 @@
 
 #include "CoreMinimal.h"
 #include "DSP/Dsp.h"
-#include "DSP/BufferVectorOperations.h"
 #include "DSP/AudioFFT.h"
+#include "DSP/BufferVectorOperations.h"
 #include "SampleBuffer.h"
 #include "Async/AsyncWork.h"
 
 namespace Audio
 {
+	class IFFTAlgorithm;
 	class FSpectrumAnalyzer;
-
 	struct SIGNALPROCESSING_API FSpectrumAnalyzerSettings
 	{
 		// Actual FFT size used. For FSpectrumAnalyzer, we never zero pad the input buffer.
@@ -30,16 +30,25 @@ namespace Audio
 
 		// Peak interpolation method. If the EFFTSize is small but will be densely sampled,
 		// it's worth using a linear or quadratic interpolation method.
-		enum class EPeakInterpolationMethod : uint16
+		enum class EPeakInterpolationMethod : uint8
 		{
 			NearestNeighbor,
 			Linear,
-			Quadratic
+			Quadratic,
+			ConstantQ,
+		};
+
+		enum class ESpectrumAnalyzerType : uint8
+		{
+			Magnitude,
+			Power,
+			Decibel,
 		};
 
 		EWindowType WindowType;
 		EFFTSize FFTSize;
 		EPeakInterpolationMethod InterpolationMethod;
+		ESpectrumAnalyzerType  SpectrumType;
 
 		/**
 			* Hop size as a percentage of FFTSize.
@@ -53,23 +62,132 @@ namespace Audio
 			: WindowType(EWindowType::Hann)
 			, FFTSize(EFFTSize::Default)
 			, InterpolationMethod(EPeakInterpolationMethod::Linear)
+			, SpectrumType(ESpectrumAnalyzerType::Magnitude)
 			, HopSize(0.0f)
 		{}
 	};
 
-	/**
-	 * This struct contains the output results from a singular FFT operation.
-	 * Stored in FSpectrumAnalyzerBuffer.
-	 */
-	struct SIGNALPROCESSING_API FSpectrumAnalyzerFrequencyVector
+	/** Settings for band extractor. */
+	struct SIGNALPROCESSING_API FSpectrumBandExtractorSettings
 	{
-		FSpectrumAnalyzerFrequencyVector(int32 InFFTSize);
+		/** Sample rate of audio */
+		float SampleRate; 
 
-		AlignedFloatBuffer RealFrequencies;
-		AlignedFloatBuffer ImagFrequencies;
+		/** Size of fft used in spectrum analyzer */
+		int32 FFTSize; 
 
-	private:
-		FSpectrumAnalyzerFrequencyVector();
+		/** Forward scaling of FFT used in spectrum analyzer */
+		EFFTScaling FFTScaling; 
+
+		/** Window used when perform FFT */
+		EWindowType WindowType;
+
+
+		/** Compare whether two settings structures are equal. */
+		bool operator==(const FSpectrumBandExtractorSettings& Other) const
+		{
+			bool bIsEqual = ((SampleRate == Other.SampleRate)
+					&& (FFTSize == Other.FFTSize)
+					&& (FFTScaling == Other.FFTScaling)
+					&& (WindowType == Other.WindowType));
+			return bIsEqual;
+		}
+
+		/** Compare whether two settings structures are not equal. */
+		bool operator!=(const FSpectrumBandExtractorSettings& Other) const
+		{
+			return !(*this == Other);
+		}
+	};
+
+	/** Interface for spectrum band extractors.
+	 *
+	 *  The SpectrumBandExtractor allows for band information
+	 *  to be maintained across multiple calls to retrieve bands values.
+	 *  By maintaining band information across multiple calls, some intermediate 
+	 *  values can be cached to speed up the operation.
+	 */
+	class SIGNALPROCESSING_API ISpectrumBandExtractor
+	{
+		public:
+			/** Metric for output band values. */
+			enum class EMetric : uint8
+			{
+				/** Return the magnitude spectrum value. */
+				Magnitude,
+
+				/** Return the power spectrum value. */
+				Power,
+
+				/** Return the decibel spectrum value. Decibels are calculated
+				 * with 0dB equal to 1.f magnitude.  */
+				Decibel
+			};
+
+			virtual ~ISpectrumBandExtractor() {}
+
+			/** Set the settings and update cached internal values if needed */
+			virtual void SetSettings(const FSpectrumBandExtractorSettings& InSettings) = 0;
+			
+			/** Removes all added bands. */
+			virtual void RemoveAllBands() = 0;
+
+			/** Returns the total number of bands. */
+			virtual int32 GetNumBands() const = 0;
+
+			/** Adds a band which calculates the band value as the value of the FFT bin nearest to the center frequency.
+			 *
+			 * @param InCenterFrequency - Frequency of interest in hz .
+			 * @param InMetric - Metric used to calculate return value.
+			 * @param InDecibelNoiseFloor - If the metric is Decibel, this is the minimum decibel value allowed.
+			 * @param bInDoNormalize - If true, all values are scaled and clamped between 0.0 and 1.f. In the 
+			 *                         case of Decibels, 0.0 corresponds to the decibel noise floor and 1.f to 0dB.
+			 */
+			virtual void AddNearestNeighborBand(float InCenterFrequency, EMetric InMetric = EMetric::Decibel, float InDecibelNoiseFloor=-40.f, bool bInDoNormalize=true) = 0;
+
+			/** Adds a band which calculates the band value as a linear interpolation of the values of the FFT 
+			 *  bins adjacent to the center frequency.
+			 *
+			 * @param InCenterFrequency - Frequency of interest in hz .
+			 * @param InMetric - Metric used to calculate return value.
+			 * @param InDecibelNoiseFloor - If the metric is Decibel, this is the minimum decibel value allowed.
+			 * @param bInDoNormalize - If true, all values are scaled and clamped between 0.0 and 1.f. In the 
+			 *                         case of Decibels, 0.0 corresponds to the decibel noise floor and 1.f to 0dB.
+			 */
+			virtual void AddLerpBand(float InCenterFrequency, EMetric InMetric = EMetric::Decibel, float InDecibelNoiseFloor=-40.f, bool bInDoNormalize=true) = 0;
+
+			/** Adds a band which calculates the band value as a quadratic interpolation of the values of the FFT 
+			 *  bins adjacent to the center frequency.
+			 *
+			 * @param InCenterFrequency - Frequency of interest in hz .
+			 * @param InMetric - Metric used to calculate return value.
+			 * @param InDecibelNoiseFloor - If the metric is Decibel, this is the minimum decibel value allowed.
+			 * @param bInDoNormalize - If true, all values are scaled and clamped between 0.0 and 1.f. In the 
+			 *                         case of Decibels, 0.0 corresponds to the decibel noise floor and 1.f to 0dB.
+			 */
+			virtual void AddQuadraticBand(float InCenterFrequency, EMetric InMetric = EMetric::Decibel, float InDecibelNoiseFloor=-40.f, bool bInDoNormalize=true) = 0;
+
+
+			/** Adds a band which calculates the band value as a pesudo constant q band derived from an FFT power spectrum.
+			 *
+			 * @param InCenterFrequency - Frequency of interest in hz.
+			 * @param InQFactor - QFactor used in band calculation. QFactor = CenterFreq / BandWidth. A small QFactor results in a wide band.
+			 * @param InMetric - Metric used to calculate return value.
+			 * @param InDecibelNoiseFloor - If the metric is Decibel, this is the minimum decibel value allowed.
+			 * @param bInDoNormalize - If true, all values are scaled and clamped between 0.0 and 1.f. In the 
+			 *                         case of Decibels, 0.0 corresponds to the decibel noise floor and 1.f to 0dB.
+			 */
+			virtual void AddConstantQBand(float InCenterFrequency, float InQFactor, EMetric InMetric = EMetric::Decibel, float InDecibelNoiseFloor=-40.f, bool bInDoNormalize=true) = 0;
+
+			/** Extract the bands from a complex frequency buffer.
+			 *
+			 * @param InComplexBuffer - Buffer of complex frequency data from a FFT.
+			 * @param OutValues - Array to store output bands.
+			 */
+			virtual void ExtractBands(const AlignedFloatBuffer& InComplexBuffer, TArray<float>& OutValues) = 0;
+
+			/** Creates a ISpectrumBandExtractor. */
+			static TUniquePtr<ISpectrumBandExtractor> CreateSpectrumBandExtractor(const FSpectrumBandExtractorSettings& InSettings);
 	};
 
 	/**
@@ -80,20 +198,20 @@ namespace Audio
 	{
 	public:
 		FSpectrumAnalyzerBuffer();
-		FSpectrumAnalyzerBuffer(const FSpectrumAnalyzerSettings& InSettings);
+		FSpectrumAnalyzerBuffer(int32 InNum);
 
-		void Reset(const FSpectrumAnalyzerSettings& InSettings);
+		void Reset(int32 InNum);
 
 		// Input. Used on analysis thread to lock a buffer to write to.
-		FSpectrumAnalyzerFrequencyVector* StartWorkOnBuffer();
+		AlignedFloatBuffer& StartWorkOnBuffer();
 		void StopWorkOnBuffer();
 		
 		// Output. Used to lock the most recent buffer we analyzed.
-		const FSpectrumAnalyzerFrequencyVector* LockMostRecentBuffer();
+		const AlignedFloatBuffer& LockMostRecentBuffer() const;
 		void UnlockBuffer();
 
 	private:
-		TArray<FSpectrumAnalyzerFrequencyVector> FrequencyVectors;
+		TArray<AlignedFloatBuffer> ComplexBuffers;
 
 		// Private functions. Either increments or decrements the respective counter,
 		// based on which index is currently in use. Mutually locked.
@@ -166,6 +284,9 @@ namespace Audio
 		// Samples phase for a given frequency, in Hz.
 		float GetPhaseForFrequency(float InFrequency);
 
+		// Return array of bands using spectrum band extractor.
+		void GetBands(ISpectrumBandExtractor& InExtractor, TArray<float>& OutValues);
+
 		// You can call this function to ensure that you're sampling the same window of frequency data,
 		// Then call UnlockOutputBuffer when you're done.
 		// Otherwise, GetMagnitudeForFrequency and GetPhaseForFrequency will always use the latest window
@@ -193,7 +314,7 @@ namespace Audio
 		void ResetSettings();
 
 		// Called in GetMagnitudeForFrequency and GetPhaseForFrequency.
-		void PerformInterpolation(const FSpectrumAnalyzerFrequencyVector* InFrequencies, FSpectrumAnalyzerSettings::EPeakInterpolationMethod InMethod, const float InFreq, float& OutReal, float& OutImag);
+		void PerformInterpolation(const AlignedFloatBuffer& InComplexBuffer, FSpectrumAnalyzerSettings::EPeakInterpolationMethod InMethod, const float InFreq, float& OutReal, float& OutImag);
 
 		// Cached current settings. Only actually used in ResetSettings().
 		FSpectrumAnalyzerSettings CurrentSettings;
@@ -207,16 +328,37 @@ namespace Audio
 		FWindow Window;
 		int32 FFTSize;
 		int32 HopInSamples;
+		EFFTScaling FFTScaling;
 
-		TArray<float> AnalysisTimeDomainBuffer;
+		AlignedFloatBuffer AnalysisTimeDomainBuffer;
 		TCircularAudioBuffer<float> InputQueue;
 		FSpectrumAnalyzerBuffer FrequencyBuffer;
 
 		// if non-null, owns pointer to locked frequency vector we're using.
-		const FSpectrumAnalyzerFrequencyVector* LockedFrequencyVector;
+		const AlignedFloatBuffer* LockedFrequencyVector;
 
 		// This is used if PerformAnalysisIfPossible is called
 		// with bAsync = true.
 		TUniquePtr<FSpectrumAnalyzerTask> AsyncAnalysisTask;
+
+		TUniquePtr<IFFTAlgorithm> FFT;
+	};
+
+	class SIGNALPROCESSING_API FSpectrumAnalyzerScopeLock
+	{
+	public:
+		FSpectrumAnalyzerScopeLock(FSpectrumAnalyzer* InAnalyzer)
+			: Analyzer(InAnalyzer)
+		{
+			Analyzer->LockOutputBuffer();
+		}
+
+		~FSpectrumAnalyzerScopeLock()
+		{
+			Analyzer->UnlockOutputBuffer();
+		}
+
+	private:
+		FSpectrumAnalyzer* Analyzer;
 	};
 }
