Fixing normalization of band extractor

diff --git a/Engine/Source/Runtime/SignalProcessing/Private/SpectrumAnalyzer.cpp b/Engine/Source/Runtime/SignalProcessing/Private/SpectrumAnalyzer.cpp
--- a/Engine/Source/Runtime/SignalProcessing/Private/SpectrumAnalyzer.cpp
+++ b/Engine/Source/Runtime/SignalProcessing/Private/SpectrumAnalyzer.cpp
@@ -85,6 +85,7 @@ namespace Audio
 							break;
 					}
 
+					/*
 					float WindowScale = 1.f;
 
 
@@ -127,6 +128,7 @@ namespace Audio
 					}
 
 					PowerSpectrumScale *= WindowScale;
+					*/
 				}
 			};
 
@@ -534,7 +536,7 @@ namespace Audio
 
 					case ISpectrumBandExtractor::EMetric::Power:
 					default:
-						OutValue = InValue;
+						OutValue = OutValue;
 						break;
 				}
 
