Fixing CIS build issue in spectrumanalyzer.cpp

diff --git a/Engine/Source/Runtime/SignalProcessing/Private/SpectrumAnalyzer.cpp b/Engine/Source/Runtime/SignalProcessing/Private/SpectrumAnalyzer.cpp
--- a/Engine/Source/Runtime/SignalProcessing/Private/SpectrumAnalyzer.cpp
+++ b/Engine/Source/Runtime/SignalProcessing/Private/SpectrumAnalyzer.cpp
@@ -536,7 +536,6 @@ namespace Audio
 
 					case ISpectrumBandExtractor::EMetric::Power:
 					default:
-						OutValue = OutValue;
 						break;
 				}
 
