Add buffer vector ops for setting a vector to constant (other than 0'ing)
#rnx

diff --git a/Engine/Source/Runtime/SignalProcessing/Private/BufferVectorOperations.cpp b/Engine/Source/Runtime/SignalProcessing/Private/BufferVectorOperations.cpp
--- a/Engine/Source/Runtime/SignalProcessing/Private/BufferVectorOperations.cpp
+++ b/Engine/Source/Runtime/SignalProcessing/Private/BufferVectorOperations.cpp
@@ -165,6 +165,21 @@ namespace Audio
 		}
 	}
 
+	void BufferSetToConstantInplace(AlignedFloatBuffer& InBuffer, float InConstant)
+	{
+		BufferSetToConstantInplace(InBuffer.GetData(), InBuffer.Num(), InConstant);
+	}
+
+	void BufferSetToConstantInplace(float* RESTRICT InBuffer, int32 NumSamples, float InConstant)
+	{
+		const VectorRegister Constant = VectorLoadFloat1(&InConstant);
+
+		for (int32 i = 0; i < NumSamples; i += 4)
+		{
+			VectorStoreAligned(Constant, &InBuffer[i]);
+		}
+	}
+
 	/* Performs an element-wise weighted sum OutputBuffer = (InBuffer1 x InGain1) + (InBuffer2 x InGain2) */
 	void BufferWeightedSumFast(const AlignedFloatBuffer& InBuffer1, float InGain1, const AlignedFloatBuffer& InBuffer2, float InGain2, AlignedFloatBuffer& OutBuffer)
 	{
diff --git a/Engine/Source/Runtime/SignalProcessing/Public/DSP/BufferVectorOperations.h b/Engine/Source/Runtime/SignalProcessing/Public/DSP/BufferVectorOperations.h
--- a/Engine/Source/Runtime/SignalProcessing/Public/DSP/BufferVectorOperations.h
+++ b/Engine/Source/Runtime/SignalProcessing/Public/DSP/BufferVectorOperations.h
@@ -33,8 +33,8 @@ namespace Audio
 	SIGNALPROCESSING_API void BufferUnderflowClampFast(float* RESTRICT InOutBuffer, const int32 InNum);
 	
 	/* Clamps the values in a buffer between a min and max value.*/
-	void SIGNALPROCESSING_API BufferRangeClampFast(AlignedFloatBuffer& InOutBuffer, float InMinValue, float InMaxValue);
-	void SIGNALPROCESSING_API BufferRangeClampFast(float* RESTRICT InOutBuffer, const int32 InNum, float InMinValue, float InMaxValue);
+	SIGNALPROCESSING_API void BufferRangeClampFast(AlignedFloatBuffer& InOutBuffer, float InMinValue, float InMaxValue);
+	SIGNALPROCESSING_API void BufferRangeClampFast(float* RESTRICT InOutBuffer, const int32 InNum, float InMinValue, float InMaxValue);
 	
 	/** Multiplies the input aligned float buffer with the given value. */
 	SIGNALPROCESSING_API void BufferMultiplyByConstant(const AlignedFloatBuffer& InFloatBuffer, float InValue, AlignedFloatBuffer& OutFloatBuffer);
@@ -49,6 +49,10 @@ namespace Audio
 	SIGNALPROCESSING_API void AddConstantToBufferInplace(AlignedFloatBuffer& InBuffer, float Constant);
 	SIGNALPROCESSING_API void AddConstantToBufferInplace(float* RESTRICT InBuffer, int32 NumSamples, float Constant);
 
+	/** Sets a constant to a buffer (useful for DC offset application) */
+	SIGNALPROCESSING_API void BufferSetToConstantInplace(AlignedFloatBuffer& InBuffer, float Constant);
+	SIGNALPROCESSING_API void BufferSetToConstantInplace(float* RESTRICT InBuffer, int32 NumSamples, float Constant);
+
 	/* Performs an element-wise weighted sum OutputBuffer = (InBuffer1 x InGain1) + (InBuffer2 x InGain2) */
 	SIGNALPROCESSING_API void BufferWeightedSumFast(const AlignedFloatBuffer& InBuffer1, float InGain1, const AlignedFloatBuffer& InBuffer2, float InGain2, AlignedFloatBuffer& OutBuffer);
 
