Fix for sounds with a short fade-in not playing. (this fix doesn't regress sequencer).

diff --git a/Engine/Source/Runtime/SignalProcessing/Private/VolumeFader.cpp b/Engine/Source/Runtime/SignalProcessing/Private/VolumeFader.cpp
--- a/Engine/Source/Runtime/SignalProcessing/Private/VolumeFader.cpp
+++ b/Engine/Source/Runtime/SignalProcessing/Private/VolumeFader.cpp
@@ -242,9 +242,14 @@ namespace Audio
 
 	void FVolumeFader::Update(float InDeltaTime)
 	{
+		// querying state before incrementing Elapsed time
+		// lets sounds with a fade-in time < InDeltaTime play
+		const bool bIsFading = IsFading();
+		const bool bIsActive = IsActive();
+
 		Elapsed += InDeltaTime;
 
-		if (!IsFading() || !IsActive())
+		if (!bIsFading || !bIsActive)
 		{
 			return;
 		}
