Adding autorange to SpectrumAnalyzer. Adding abilitiy to automatically control dynamic range of SoundSubmix being analyzed.

diff --git a/Engine/Source/Runtime/AudioMixer/Private/AudioMixerBlueprintLibrary.cpp b/Engine/Source/Runtime/AudioMixer/Private/AudioMixerBlueprintLibrary.cpp
--- a/Engine/Source/Runtime/AudioMixer/Private/AudioMixerBlueprintLibrary.cpp
+++ b/Engine/Source/Runtime/AudioMixer/Private/AudioMixerBlueprintLibrary.cpp
@@ -266,7 +266,7 @@ void UAudioMixerBlueprintLibrary::StartAnalyzingOutput(const UObject* WorldConte
 {
 	if (Audio::FMixerDevice* MixerDevice = GetAudioMixerDeviceFromWorldContext(WorldContextObject))
 	{
-		Audio::FSpectrumAnalyzerSettings Settings = USoundSubmix::GetSpectrumAnalyzerSettings(FFTSize, InterpolationMethod, WindowType, HopSize, AudioSpectrumType);
+		FSoundSpectrumAnalyzerSettings Settings = USoundSubmix::GetSpectrumAnalyzerSettings(FFTSize, InterpolationMethod, WindowType, HopSize, AudioSpectrumType);
 		MixerDevice->StartSpectrumAnalysis(SubmixToAnalyze, Settings);
 	}
 	else
@@ -401,7 +401,7 @@ void UAudioMixerBlueprintLibrary::GetMagnitudeForFrequencies(const UObject* Worl
 	}
 	else
 	{
-		UE_LOG(LogAudioMixer, Error, TEXT("Output recording is an audio mixer only feature."));
+		UE_LOG(LogAudioMixer, Error, TEXT("Getting magnitude for frequencies is an audio mixer only feature."));
 	}
 }
 
diff --git a/Engine/Source/Runtime/AudioMixer/Private/AudioMixerDevice.cpp b/Engine/Source/Runtime/AudioMixer/Private/AudioMixerDevice.cpp
--- a/Engine/Source/Runtime/AudioMixer/Private/AudioMixerDevice.cpp
+++ b/Engine/Source/Runtime/AudioMixer/Private/AudioMixerDevice.cpp
@@ -1818,7 +1818,7 @@ namespace Audio
 	}
 
 
-	void FMixerDevice::StartSpectrumAnalysis(USoundSubmix* InSubmix, const Audio::FSpectrumAnalyzerSettings& InSettings)
+	void FMixerDevice::StartSpectrumAnalysis(USoundSubmix* InSubmix, const FSoundSpectrumAnalyzerSettings& InSettings)
 	{
 		if (!IsInAudioThread())
 		{
@@ -1915,16 +1915,17 @@ namespace Audio
 		}
 	}
 
-	void FMixerDevice::AddSpectralAnalysisDelegate(USoundSubmix* InSubmix, const TArray<FSoundSubmixSpectralAnalysisBandSettings>& InBandSettings, const FOnSubmixSpectralAnalysisBP& OnSubmixSpectralAnalysisBP, float UpdateRate)
+	void FMixerDevice::AddSpectralAnalysisDelegate(USoundSubmix* InSubmix, const FSoundSpectrumAnalyzerDelegateSettings& InDelegateSettings, const FOnSubmixSpectralAnalysisBP& OnSubmixSpectralAnalysisBP)
 	{
+
 		if (!IsInAudioThread())
 		{
 			DECLARE_CYCLE_STAT(TEXT("FAudioThreadTask.AddSpectralAnalysisDelegate"), STAT_AddSpectralAnalysisDelegate, STATGROUP_AudioThreadCommands);
 
-			FAudioThread::RunCommandOnAudioThread([this, InSubmix, InBandSettings, OnSubmixSpectralAnalysisBP, UpdateRate]()
+			FAudioThread::RunCommandOnAudioThread([this, InSubmix, InDelegateSettings, OnSubmixSpectralAnalysisBP]()
 			{
 				CSV_SCOPED_TIMING_STAT(Audio, AddSpectralAnalysisDelegate);
-				AddSpectralAnalysisDelegate(InSubmix, InBandSettings, OnSubmixSpectralAnalysisBP, UpdateRate);
+				AddSpectralAnalysisDelegate(InSubmix, InDelegateSettings, OnSubmixSpectralAnalysisBP);
 			}, GET_STATID(STAT_AddSpectralAnalysisDelegate));
 			return;
 		}
@@ -1941,7 +1942,7 @@ namespace Audio
 
 		if (ensure(FoundSubmix.IsValid()))
 		{
-			FoundSubmix->AddSpectralAnalysisDelegate(OnSubmixSpectralAnalysisBP, InBandSettings, UpdateRate);
+			FoundSubmix->AddSpectralAnalysisDelegate(InDelegateSettings, OnSubmixSpectralAnalysisBP);
 		}
 	}
 
diff --git a/Engine/Source/Runtime/AudioMixer/Private/AudioMixerSubmix.cpp b/Engine/Source/Runtime/AudioMixer/Private/AudioMixerSubmix.cpp
--- a/Engine/Source/Runtime/AudioMixer/Private/AudioMixerSubmix.cpp
+++ b/Engine/Source/Runtime/AudioMixer/Private/AudioMixerSubmix.cpp
@@ -30,6 +30,120 @@ FAutoConsoleVariableRef CVarBypassAllSubmixEffects(
 
 namespace Audio
 {
+	namespace MixerSubmixIntrinsics
+	{
+
+		FSpectrumAnalyzerSettings::EFFTSize GetSpectrumAnalyzerFFTSize(EFFTSize InFFTSize)
+		{
+			switch (InFFTSize)
+			{
+				case EFFTSize::DefaultSize:
+					return FSpectrumAnalyzerSettings::EFFTSize::Default;
+					break;
+
+				case EFFTSize::Min:
+					return FSpectrumAnalyzerSettings::EFFTSize::Min_64;
+					break;
+
+				case EFFTSize::Small:
+					return FSpectrumAnalyzerSettings::EFFTSize::Small_256;
+					break;
+
+				case EFFTSize::Medium:
+					return FSpectrumAnalyzerSettings::EFFTSize::Medium_512;
+					break;
+
+				case EFFTSize::Large:
+					return FSpectrumAnalyzerSettings::EFFTSize::Large_1024;
+					break;
+
+				case EFFTSize::VeryLarge:
+					return FSpectrumAnalyzerSettings::EFFTSize::VeryLarge_2048;
+					break;
+
+				case EFFTSize::Max:
+					return FSpectrumAnalyzerSettings::EFFTSize::TestLarge_4096;
+					break;
+
+				default:
+					return FSpectrumAnalyzerSettings::EFFTSize::Default;
+					break;
+			}
+		}
+
+		EWindowType GetWindowType(EFFTWindowType InWindowType)
+		{
+			switch (InWindowType)
+			{
+				case EFFTWindowType::None:
+					return EWindowType::None;
+					break;
+
+				case EFFTWindowType::Hamming:
+					return EWindowType::Hamming;
+					break;
+
+				case EFFTWindowType::Hann:
+					return EWindowType::Hann;
+					break;
+
+				case EFFTWindowType::Blackman:
+					return EWindowType::Blackman;
+					break;
+
+				default:
+					return EWindowType::None;
+					break;
+			}
+		}
+
+		FSpectrumBandExtractorSettings::EMetric GetExtractorMetric(EAudioSpectrumType InSpectrumType)
+		{
+			using EMetric = FSpectrumBandExtractorSettings::EMetric;
+
+			switch (InSpectrumType)
+			{
+				case EAudioSpectrumType::MagnitudeSpectrum:
+					return EMetric::Magnitude;
+					break;
+
+				case EAudioSpectrumType::PowerSpectrum:
+					return EMetric::Power;
+					break;
+
+				case EAudioSpectrumType::Decibel:
+				default:
+					return EMetric::Decibel;
+					break;
+			}
+		}
+
+		ISpectrumBandExtractor::EBandType GetExtractorBandType(EFFTPeakInterpolationMethod InMethod)
+		{
+			using EBandType = ISpectrumBandExtractor::EBandType;
+
+			switch (InMethod)
+			{
+				case EFFTPeakInterpolationMethod::NearestNeighbor:
+					return EBandType::NearestNeighbor;
+					break;
+
+				case EFFTPeakInterpolationMethod::Linear:
+					return EBandType::Lerp;
+					break;
+
+				case EFFTPeakInterpolationMethod::Quadratic:
+					return EBandType::Quadratic;
+					break;
+
+				case EFFTPeakInterpolationMethod::ConstantQ:
+				default:
+					return EBandType::ConstantQ;
+					break;
+			}
+		}
+	}
+
 	// Unique IDs for mixer submixes
 	static uint32 GSubmixMixerIDs = 0;
 
@@ -1478,26 +1592,16 @@ namespace Audio
 		OnSubmixEnvelope.AddUnique(OnSubmixEnvelopeBP);
 	}
 
-	void FMixerSubmix::AddSpectralAnalysisDelegate(const FOnSubmixSpectralAnalysisBP& OnSubmixSpectralAnalysisBP, const TArray<FSoundSubmixSpectralAnalysisBandSettings>& InBandSettings, float UpdateRate)
+	void FMixerSubmix::AddSpectralAnalysisDelegate(const FSoundSpectrumAnalyzerDelegateSettings& InDelegateSettings, const FOnSubmixSpectralAnalysisBP& OnSubmixSpectralAnalysisBP)
 	{
 		FSpectrumAnalysisDelegateInfo& NewDelegateInfo = SpectralAnalysisDelegates.AddDefaulted_GetRef();
-		
+	
 		NewDelegateInfo.LastUpdateTime = -1.0f;
-		UpdateRate = FMath::Clamp(UpdateRate, 1.0f, 30.0f);
-		NewDelegateInfo.UpdateDelta = 1.0f / UpdateRate;
-
-		for (const FSoundSubmixSpectralAnalysisBandSettings& BandSettings : InBandSettings)
-		{
-			FSpectralAnalysisBandInfo NewBand;
-			NewBand.BandFrequency = BandSettings.BandFrequency;
-			NewBand.QFactor = BandSettings.QFactor;
-			NewBand.EnvelopeFollower.Init(UpdateRate, BandSettings.AttackTimeMsec, BandSettings.ReleaseTimeMsec);
-		
-			NewDelegateInfo.SpectralBands.Add(NewBand);
-		}
+		NewDelegateInfo.DelegateSettings = InDelegateSettings;
+		NewDelegateInfo.DelegateSettings.UpdateRate = FMath::Clamp(NewDelegateInfo.DelegateSettings.UpdateRate, 1.0f, 30.0f);
+		NewDelegateInfo.UpdateDelta = 1.0f / NewDelegateInfo.DelegateSettings.UpdateRate;
 
 		NewDelegateInfo.OnSubmixSpectralAnalysis.AddUnique(OnSubmixSpectralAnalysisBP);
-
 	}
 
 	void FMixerSubmix::RemoveSpectralAnalysisDelegate(const FOnSubmixSpectralAnalysisBP& OnSubmixSpectralAnalysisBP)
@@ -1515,59 +1619,58 @@ namespace Audio
 		});
 	}
 
-	void FMixerSubmix::StartSpectrumAnalysis(const FSpectrumAnalyzerSettings& InSettings)
+	void FMixerSubmix::StartSpectrumAnalysis(const FSoundSpectrumAnalyzerSettings& InSettings)
 	{
+		using namespace MixerSubmixIntrinsics;
+		using EMetric = FSpectrumBandExtractorSettings::EMetric;
+		using EBandType = ISpectrumBandExtractor::EBandType;
+
 		bIsSpectrumAnalyzing = true;
 
-		// Construct a band extractor for each delegate listening to FFT results
+		SpectrumAnalyzerSettings = InSettings;
+
+		FSpectrumAnalyzerSettings AudioSpectrumAnalyzerSettings;
+
+		AudioSpectrumAnalyzerSettings.FFTSize = GetSpectrumAnalyzerFFTSize(SpectrumAnalyzerSettings.FFTSize);
+		AudioSpectrumAnalyzerSettings.WindowType = GetWindowType(SpectrumAnalyzerSettings.WindowType);
+		AudioSpectrumAnalyzerSettings.HopSize = SpectrumAnalyzerSettings.HopSize;
+
+		SpectrumAnalyzer.Reset(new FSpectrumAnalyzer(AudioSpectrumAnalyzerSettings, MixerDevice->GetSampleRate()));
+
+		EMetric Metric = GetExtractorMetric(SpectrumAnalyzerSettings.SpectrumType);
+		EBandType BandType = GetExtractorBandType(SpectrumAnalyzerSettings.InterpolationMethod);
+
 		for (FSpectrumAnalysisDelegateInfo& DelegateInfo : SpectralAnalysisDelegates)
 		{
-			// Build the band extractor here since we are given the FFT settings to use at this point
-			FSpectrumBandExtractorSettings BandExtractorSettings;
+			FSpectrumBandExtractorSettings ExtractorSettings;
 
-			BandExtractorSettings.FFTSize = (int32)InSettings.FFTSize;
-			BandExtractorSettings.SampleRate = MixerDevice->GetSampleRate();
+			ExtractorSettings.Metric = Metric;
+			ExtractorSettings.DecibelNoiseFloor = DelegateInfo.DelegateSettings.DecibelNoiseFloor;
+			ExtractorSettings.bDoNormalize = DelegateInfo.DelegateSettings.bDoNormalize;
+			ExtractorSettings.bDoAutoRange = DelegateInfo.DelegateSettings.bDoAutoRange;
+			ExtractorSettings.AutoRangeReleaseTimeInSeconds = DelegateInfo.DelegateSettings.AutoRangeReleaseTime;
+			ExtractorSettings.AutoRangeAttackTimeInSeconds = DelegateInfo.DelegateSettings.AutoRangeAttackTime;
 
-			DelegateInfo.SpectrumBandExtractor = ISpectrumBandExtractor::CreateSpectrumBandExtractor(BandExtractorSettings);
+			DelegateInfo.SpectrumBandExtractor = ISpectrumBandExtractor::CreateSpectrumBandExtractor(ExtractorSettings);
 
-			// Determine the band extractor metric (i.e. spectrum type). This is the same for all bands in a given extractor.
-			ISpectrumBandExtractor::EMetric Metric = ISpectrumBandExtractor::EMetric::Magnitude;
-			switch (InSettings.SpectrumType)
+			if (DelegateInfo.SpectrumBandExtractor.IsValid())
 			{
-			case FSpectrumAnalyzerSettings::ESpectrumAnalyzerType::Power:
-				Metric = ISpectrumBandExtractor::EMetric::Power;
-				break;
+				for (const FSoundSubmixSpectralAnalysisBandSettings& BandSettings : DelegateInfo.DelegateSettings.BandSettings)
+				{
+					ISpectrumBandExtractor::FBandSettings NewExtractorBandSettings;
+					NewExtractorBandSettings.Type = BandType;
+					NewExtractorBandSettings.CenterFrequency = BandSettings.BandFrequency;
+					NewExtractorBandSettings.QFactor = BandSettings.QFactor;
 
-			case FSpectrumAnalyzerSettings::ESpectrumAnalyzerType::Decibel:
-				Metric = ISpectrumBandExtractor::EMetric::Decibel;
-				break;
+					DelegateInfo.SpectrumBandExtractor->AddBand(NewExtractorBandSettings);
 
-			default:
-				break;
-			}
-			
-			// Loop through each band and add them to the extractor
-			for (FSpectralAnalysisBandInfo& BandInfo : DelegateInfo.SpectralBands)
-			{
-				switch (InSettings.InterpolationMethod)
-				{
-					case FSpectrumAnalyzerSettings::EPeakInterpolationMethod::NearestNeighbor:
-						DelegateInfo.SpectrumBandExtractor->AddNearestNeighborBand(BandInfo.BandFrequency, Metric);
-						break;
-					case FSpectrumAnalyzerSettings::EPeakInterpolationMethod::Linear:
-						DelegateInfo.SpectrumBandExtractor->AddLerpBand(BandInfo.BandFrequency, Metric);
-						break;
-					case FSpectrumAnalyzerSettings::EPeakInterpolationMethod::Quadratic:
-						DelegateInfo.SpectrumBandExtractor->AddQuadraticBand(BandInfo.BandFrequency, Metric);
-						break;
-					case FSpectrumAnalyzerSettings::EPeakInterpolationMethod::ConstantQ:
-						DelegateInfo.SpectrumBandExtractor->AddConstantQBand(BandInfo.BandFrequency, BandInfo.QFactor, Metric);
-						break;
+					FSpectralAnalysisBandInfo NewBand;
+					NewBand.EnvelopeFollower.Init(DelegateInfo.DelegateSettings.UpdateRate, BandSettings.AttackTimeMsec, BandSettings.ReleaseTimeMsec);
+				
+					DelegateInfo.SpectralBands.Add(NewBand);
 				}
-			}
+			} 
 		}
-
-		SpectrumAnalyzer.Reset(new FSpectrumAnalyzer(InSettings, MixerDevice->GetSampleRate()));
 	}
 
 	void FMixerSubmix::StopSpectrumAnalysis()
@@ -1580,13 +1683,36 @@ namespace Audio
 	{
 		if (SpectrumAnalyzer.IsValid())
 		{
+			using EMethod = FSpectrumAnalyzer::EPeakInterpolationMethod;
+
+			EMethod Method;	
+
+			switch (SpectrumAnalyzerSettings.InterpolationMethod)
+			{
+				case EFFTPeakInterpolationMethod::NearestNeighbor:
+					Method = EMethod::NearestNeighbor;
+					break;
+
+				case EFFTPeakInterpolationMethod::Linear:
+					Method = EMethod::Linear;
+					break;
+
+				case EFFTPeakInterpolationMethod::Quadratic:
+					Method = EMethod::Quadratic;
+					break;
+
+				default:
+					Method = EMethod::Linear;
+					break;
+			}
+
 			OutMagnitudes.Reset();
 			OutMagnitudes.AddUninitialized(InFrequencies.Num());
 
 			SpectrumAnalyzer->LockOutputBuffer();
 			for (int32 Index = 0; Index < InFrequencies.Num(); Index++)
 			{
-				OutMagnitudes[Index] = SpectrumAnalyzer->GetMagnitudeForFrequency(InFrequencies[Index]);
+				OutMagnitudes[Index] = SpectrumAnalyzer->GetMagnitudeForFrequency(InFrequencies[Index], Method);
 			}
 			SpectrumAnalyzer->UnlockOutputBuffer();
 		}
@@ -1600,13 +1726,36 @@ namespace Audio
 	{
 		if (SpectrumAnalyzer.IsValid())
 		{
+			using EMethod = FSpectrumAnalyzer::EPeakInterpolationMethod;
+
+			EMethod Method;	
+
+			switch (SpectrumAnalyzerSettings.InterpolationMethod)
+			{
+				case EFFTPeakInterpolationMethod::NearestNeighbor:
+					Method = EMethod::NearestNeighbor;
+					break;
+
+				case EFFTPeakInterpolationMethod::Linear:
+					Method = EMethod::Linear;
+					break;
+
+				case EFFTPeakInterpolationMethod::Quadratic:
+					Method = EMethod::Quadratic;
+					break;
+
+				default:
+					Method = EMethod::Linear;
+					break;
+			}
+
 			OutPhases.Reset();
 			OutPhases.AddUninitialized(InFrequencies.Num());
 
 			SpectrumAnalyzer->LockOutputBuffer();
 			for (int32 Index = 0; Index < InFrequencies.Num(); Index++)
 			{
-				OutPhases[Index] = SpectrumAnalyzer->GetPhaseForFrequency(InFrequencies[Index]);
+				OutPhases[Index] = SpectrumAnalyzer->GetPhaseForFrequency(InFrequencies[Index], Method);
 			}
 			SpectrumAnalyzer->UnlockOutputBuffer();
 		}
diff --git a/Engine/Source/Runtime/AudioMixer/Public/AudioMixerDevice.h b/Engine/Source/Runtime/AudioMixer/Public/AudioMixerDevice.h
--- a/Engine/Source/Runtime/AudioMixer/Public/AudioMixerDevice.h
+++ b/Engine/Source/Runtime/AudioMixer/Public/AudioMixerDevice.h
@@ -126,12 +126,13 @@ namespace Audio
 		virtual void AddEnvelopeFollowerDelegate(USoundSubmix* InSubmix, const FOnSubmixEnvelopeBP& OnSubmixEnvelopeBP) override;
 
 		// Submix Spectrum Analysis
-		virtual void StartSpectrumAnalysis(USoundSubmix* InSubmix, const Audio::FSpectrumAnalyzerSettings& InSettings) override;
+
+		virtual void StartSpectrumAnalysis(USoundSubmix* InSubmix, const FSoundSpectrumAnalyzerSettings& InSettings) override;
 		virtual void StopSpectrumAnalysis(USoundSubmix* InSubmix) override;
-		virtual void GetMagnitudesForFrequencies(USoundSubmix* InSubmix, const TArray<float>& InFrequencies, TArray<float>& OutMagnitudes);
-		virtual void GetPhasesForFrequencies(USoundSubmix* InSubmix, const TArray<float>& InFrequencies, TArray<float>& OutPhases);
-		virtual void AddSpectralAnalysisDelegate(USoundSubmix* InSubmix, const TArray<FSoundSubmixSpectralAnalysisBandSettings>& InBandSettings, const FOnSubmixSpectralAnalysisBP& OnSubmixSpectralAnalysisBP, float UpdateRate) override;
-		virtual void RemoveSpectralAnalysisDelegate(USoundSubmix* InSubmixconst, const FOnSubmixSpectralAnalysisBP& OnSubmixSpectralAnalysisBP) override;
+		virtual void GetMagnitudesForFrequencies(USoundSubmix* InSubmix, const TArray<float>& InFrequencies, TArray<float>& OutMagnitudes) override;
+		virtual void GetPhasesForFrequencies(USoundSubmix* InSubmix, const TArray<float>& InFrequencies, TArray<float>& OutPhases) override;
+		virtual void AddSpectralAnalysisDelegate(USoundSubmix* InSubmix, const FSoundSpectrumAnalyzerDelegateSettings& InDelegateSettings, const FOnSubmixSpectralAnalysisBP& OnSubmixSpectralAnalysisBP) override;
+		virtual void RemoveSpectralAnalysisDelegate(USoundSubmix* InSubmix, const FOnSubmixSpectralAnalysisBP& OnSubmixSpectralAnalysisBP) override;
 
 		// Submix buffer listener callbacks
 		virtual void RegisterSubmixBufferListener(ISubmixBufferListener* InSubmixBufferListener, USoundSubmix* InSubmix = nullptr) override;
diff --git a/Engine/Source/Runtime/AudioMixer/Public/AudioMixerSubmix.h b/Engine/Source/Runtime/AudioMixer/Public/AudioMixerSubmix.h
--- a/Engine/Source/Runtime/AudioMixer/Public/AudioMixerSubmix.h
+++ b/Engine/Source/Runtime/AudioMixer/Public/AudioMixerSubmix.h
@@ -195,13 +195,13 @@ namespace Audio
 		void AddEnvelopeFollowerDelegate(const FOnSubmixEnvelopeBP& OnSubmixEnvelopeBP);
 
 		// Initializes a new FFT analyzer for this submix and immediately begins feeding audio to it.
-		void StartSpectrumAnalysis(const FSpectrumAnalyzerSettings& InSettings);
+		void StartSpectrumAnalysis(const FSoundSpectrumAnalyzerSettings& InSettings);
 
 		// Terminates whatever FFT Analyzer is being used for this submix.
 		void StopSpectrumAnalysis();
 
 		// Adds an spectral analysis delegate
-		void AddSpectralAnalysisDelegate(const FOnSubmixSpectralAnalysisBP& OnSubmixSpectralAnalysisBP, const TArray<FSoundSubmixSpectralAnalysisBandSettings>& InBandSettings, float UpdateRate);
+		void AddSpectralAnalysisDelegate(const FSoundSpectrumAnalyzerDelegateSettings& InDelegateSettings, const FOnSubmixSpectralAnalysisBP& OnSubmixSpectralAnalysisBP);
 
 		// Removes an existing spectral analysis delegate
 		void RemoveSpectralAnalysisDelegate(const FOnSubmixSpectralAnalysisBP& OnSubmixSpectralAnalysisBP);
@@ -452,6 +452,7 @@ namespace Audio
 		FCriticalSection EnvelopeCriticalSection;
 
 		// Spectrum analyzer:
+		FSoundSpectrumAnalyzerSettings SpectrumAnalyzerSettings;
 		TUniquePtr<FSpectrumAnalyzer> SpectrumAnalyzer;
 		
 		// This buffer is used to downmix the submix output to mono before submitting it to the SpectrumAnalyzer.
@@ -490,13 +491,14 @@ namespace Audio
 		struct FSpectralAnalysisBandInfo
 		{
 			FEnvelopeFollower EnvelopeFollower;
-			float BandFrequency = 0.0f;
-			float QFactor = 1.f;
 		};
 
 		struct FSpectrumAnalysisDelegateInfo
 		{
+			FSoundSpectrumAnalyzerDelegateSettings DelegateSettings;
+
 			FOnSubmixSpectralAnalysis OnSubmixSpectralAnalysis;
+
 			TUniquePtr<ISpectrumBandExtractor> SpectrumBandExtractor;
 			TArray<FSpectralAnalysisBandInfo> SpectralBands;
 
@@ -511,6 +513,7 @@ namespace Audio
 			{
 				OnSubmixSpectralAnalysis = Other.OnSubmixSpectralAnalysis;
 				SpectrumBandExtractor.Reset(Other.SpectrumBandExtractor.Release());
+				DelegateSettings = Other.DelegateSettings;
 				SpectralBands = Other.SpectralBands;
 			}
 
diff --git a/Engine/Source/Runtime/Engine/Classes/Sound/SoundSubmix.h b/Engine/Source/Runtime/Engine/Classes/Sound/SoundSubmix.h
--- a/Engine/Source/Runtime/Engine/Classes/Sound/SoundSubmix.h
+++ b/Engine/Source/Runtime/Engine/Classes/Sound/SoundSubmix.h
@@ -101,6 +101,47 @@ enum class EAudioSpectrumType : uint8
 	Decibel,
 };
 
+struct FSoundSpectrumAnalyzerSettings
+{
+	// FFTSize used in spectrum analyzer.
+	EFFTSize FFTSize;
+
+	// Type of window to apply to audio.
+	EFFTWindowType WindowType;
+
+	// Metric used when analyzing spectrum. 
+	EAudioSpectrumType SpectrumType;
+
+	// Interpolation method used when getting frequencies.
+	EFFTPeakInterpolationMethod  InterpolationMethod;
+
+	// Hopsize between audio windows as a ratio of the FFTSize.
+	float HopSize;
+};
+
+struct FSoundSpectrumAnalyzerDelegateSettings
+{
+	// Settings for individual bands.
+	TArray<FSoundSubmixSpectralAnalysisBandSettings> BandSettings; 
+
+	// Number of times a second the delegate is triggered. 
+	float UpdateRate; 
+
+	// The decibel level considered silence.
+	float DecibelNoiseFloor; 
+
+	// If true, returned values are scaled between 0 and 1.
+	bool bDoNormalize; 
+
+	// If true, the band values are tracked to always have values between 0 and 1. 
+	bool bDoAutoRange; 
+
+	// The time in seconds for the range to expand to a new observed range.
+	float AutoRangeAttackTime; 
+
+	// The time in seconds for the range to shrink to a new observed range.
+	float AutoRangeReleaseTime;
+};
 
 
 #if WITH_EDITOR
@@ -302,9 +343,16 @@ public:
 	 *	@param	InBandsettings					The frequency bands to analyze and their envelope-following settings.
 	 *  @param  OnSubmixSpectralAnalysisBP		Event to fire when new spectral data is available.
 	 *	@param	UpdateRate						How often to retrieve the data from the spectral analyzer and broadcast the event. Max is 30 times per second.
+	 *	@param  InterpMethod                    Method to used for band peak calculation.
+	 *	@param  SpectrumType                    Metric to use when returning spectrum values.
+	 *	@param  DecibelNoiseFloor               Decibel Noise Floor to consider as silence silence when using a Decibel Spectrum Type.
+	 *	@param  bDoNormalize                    If true, output band values will be normalized between zero and one.
+	 *	@param  bDoAutoRange                    If true, output band values will have their ranges automatically adjusted to the minimum and maximum values in the audio. Output band values will be normalized between zero and one.
+	 *	@param  AutoRangeAttackTime             The time (in seconds) it takes for the range to expand to 90% of a larger range.
+	 *	@param  AutoRangeReleaseTime            The time (in seconds) it takes for the range to shrink to 90% of a smaller range.
 	 */
 	UFUNCTION(BlueprintCallable, Category = "Audio|Spectrum", meta = (WorldContext = "WorldContextObject"))
-	void AddSpectralAnalysisDelegate(const UObject* WorldContextObject, const TArray<FSoundSubmixSpectralAnalysisBandSettings>& InBandSettings, const FOnSubmixSpectralAnalysisBP& OnSubmixSpectralAnalysisBP, float UpdateRate);
+	void AddSpectralAnalysisDelegate(const UObject* WorldContextObject, const TArray<FSoundSubmixSpectralAnalysisBandSettings>& InBandSettings, const FOnSubmixSpectralAnalysisBP& OnSubmixSpectralAnalysisBP, float UpdateRate, float DecibelNoiseFloor=-40.f, bool bDoNormalize = true, bool bDoAutoRange = false, float AutoRangeAttackTime = 0.1f, float AutoRangeReleaseTime = 60.f);
 
 	/**
 	 *	Remove a spectral analysis delegate.
@@ -329,7 +377,9 @@ public:
 	UFUNCTION(BlueprintCallable, Category = "Audio", meta = (WorldContext = "WorldContextObject"))
 	void SetSubmixOutputVolume(const UObject* WorldContextObject, float InOutputVolume);
 
-	static Audio::FSpectrumAnalyzerSettings GetSpectrumAnalyzerSettings(EFFTSize FFTSize, EFFTPeakInterpolationMethod InterpolationMethod, EFFTWindowType WindowType, float HopSize, EAudioSpectrumType AudioSpectrumType);
+	static FSoundSpectrumAnalyzerSettings GetSpectrumAnalyzerSettings(EFFTSize FFTSize, EFFTPeakInterpolationMethod InterpolationMethod, EFFTWindowType WindowType, float HopSize, EAudioSpectrumType SpectrumType);
+
+	static FSoundSpectrumAnalyzerDelegateSettings GetSpectrumAnalysisDelegateSettings(const TArray<FSoundSubmixSpectralAnalysisBandSettings>& InBandSettings, float UpdateRate, float DecibelNoiseFloor, bool bDoNormalize, bool bDoAutoRange, float AutoRangeAttackTime, float AutoRangeReleaseTime);
 protected:
 
 	virtual void PostLoad() override;
diff --git a/Engine/Source/Runtime/Engine/Private/AudioDevice.cpp b/Engine/Source/Runtime/Engine/Private/AudioDevice.cpp
--- a/Engine/Source/Runtime/Engine/Private/AudioDevice.cpp
+++ b/Engine/Source/Runtime/Engine/Private/AudioDevice.cpp
@@ -4675,7 +4675,7 @@ void FAudioDevice::AddEnvelopeFollowerDelegate(USoundSubmix* InSubmix, const FOn
 	UE_LOG(LogAudio, Error, TEXT("Envelope following submixes only works with the audio mixer. Please run using -audiomixer or set INI file to use submix recording."));
 }
 
-void FAudioDevice::StartSpectrumAnalysis(USoundSubmix* InSubmix, const Audio::FSpectrumAnalyzerSettings& InSettings)
+void FAudioDevice::StartSpectrumAnalysis(USoundSubmix* InSubmix, const FSoundSpectrumAnalyzerSettings& InSettings)
 {
 	UE_LOG(LogAudio, Error, TEXT("Spectrum analysis of submixes only works with the audio mixer. Please run using -audiomixer or set INI file to use submix recording."));
 }
@@ -4695,7 +4695,7 @@ void FAudioDevice::GetPhasesForFrequencies(USoundSubmix* InSubmix, const TArray<
 	UE_LOG(LogAudio, Error, TEXT("Spectrum analysis of submixes only works with the audio mixer. Please run using -audiomixer or set INI file to use submix recording."));
 }
 
-void FAudioDevice::AddSpectralAnalysisDelegate(USoundSubmix* InSubmix, const TArray<FSoundSubmixSpectralAnalysisBandSettings>& InBandSettings, const FOnSubmixSpectralAnalysisBP& OnSubmixSpectralAnalysisBP, float UpdateRate)
+void FAudioDevice::AddSpectralAnalysisDelegate(USoundSubmix* InSubmix, const FSoundSpectrumAnalyzerDelegateSettings& InDelegateSettings, const FOnSubmixSpectralAnalysisBP& OnSubmixSpectralAnalysisBP)
 {
 	UE_LOG(LogAudio, Error, TEXT("Spectrum analysis of submixes only works with the audio mixer. Please run using -audiomixer or set INI file to use submix recording."));
 }
@@ -6527,4 +6527,4 @@ int32 FAudioDevice::GetNumPrecacheFrames() const
 	}
 	// Otherwise, use the default value or value set in ini file
 	return NumPrecacheFrames;
-}
\ No newline at end of file
+}
diff --git a/Engine/Source/Runtime/Engine/Private/SoundSubmix.cpp b/Engine/Source/Runtime/Engine/Private/SoundSubmix.cpp
--- a/Engine/Source/Runtime/Engine/Private/SoundSubmix.cpp
+++ b/Engine/Source/Runtime/Engine/Private/SoundSubmix.cpp
@@ -259,8 +259,9 @@ void USoundSubmix::AddEnvelopeFollowerDelegate(const UObject* WorldContextObject
 	}
 }
 
-void USoundSubmix::AddSpectralAnalysisDelegate(const UObject* WorldContextObject, const TArray<FSoundSubmixSpectralAnalysisBandSettings>& InBandSettings, const FOnSubmixSpectralAnalysisBP& OnSubmixSpectralAnalysisBP, float UpdateRate)
+void USoundSubmix::AddSpectralAnalysisDelegate(const UObject* WorldContextObject, const TArray<FSoundSubmixSpectralAnalysisBandSettings>& InBandSettings, const FOnSubmixSpectralAnalysisBP& OnSubmixSpectralAnalysisBP, float UpdateRate,  float DecibelNoiseFloor, bool bDoNormalize, bool bDoAutoRange, float AutoRangeAttackTime, float AutoRangeReleaseTime)
 {
+
 	if (!GEngine)
 	{
 		return;
@@ -271,7 +272,10 @@ void USoundSubmix::AddSpectralAnalysisDelegate(const UObject* WorldContextObject
 	{
 		if (FAudioDevice* AudioDevice = ThisWorld->GetAudioDeviceRaw())
 		{
-			AudioDevice->AddSpectralAnalysisDelegate(this, InBandSettings, OnSubmixSpectralAnalysisBP, UpdateRate);
+			FSoundSpectrumAnalyzerDelegateSettings DelegateSettings = USoundSubmix::GetSpectrumAnalysisDelegateSettings(InBandSettings, UpdateRate, DecibelNoiseFloor, bDoNormalize, bDoAutoRange, AutoRangeAttackTime, AutoRangeReleaseTime);
+
+
+			AudioDevice->AddSpectralAnalysisDelegate(this, DelegateSettings, OnSubmixSpectralAnalysisBP);
 		}
 	}
 }
@@ -319,7 +323,7 @@ void USoundSubmix::StartSpectralAnalysis(FAudioDevice* InAudioDevice, EFFTSize F
 
 	if (InAudioDevice)
 	{
-		Audio::FSpectrumAnalyzerSettings Settings = USoundSubmix::GetSpectrumAnalyzerSettings(FFTSize, InterpolationMethod, WindowType, HopSize, SpectrumType);
+		FSoundSpectrumAnalyzerSettings Settings = USoundSubmix::GetSpectrumAnalyzerSettings(FFTSize, InterpolationMethod, WindowType, HopSize, SpectrumType);
 		InAudioDevice->StartSpectrumAnalysis(this, Settings);
 	}
 }
@@ -885,88 +889,32 @@ TArray<USoundfieldEffectBase*> USoundfieldEndpointSubmix::GetSoundfieldProcessor
 }
 
 
-Audio::FSpectrumAnalyzerSettings USoundSubmix::GetSpectrumAnalyzerSettings(EFFTSize FFTSize, EFFTPeakInterpolationMethod InterpolationMethod, EFFTWindowType WindowType, float HopSize, EAudioSpectrumType AudioSpectrumType)
+FSoundSpectrumAnalyzerSettings USoundSubmix::GetSpectrumAnalyzerSettings(EFFTSize FFTSize, EFFTPeakInterpolationMethod InterpolationMethod, EFFTWindowType WindowType, float HopSize, EAudioSpectrumType SpectrumType)
 {
-	Audio::FSpectrumAnalyzerSettings OutSettings;
+	FSoundSpectrumAnalyzerSettings OutSettings;
 
-	switch (FFTSize)
-	{
-	case EFFTSize::DefaultSize:
-		OutSettings.FFTSize = Audio::FSpectrumAnalyzerSettings::EFFTSize::Default;
-		break;
-	case EFFTSize::Min:
-		OutSettings.FFTSize = Audio::FSpectrumAnalyzerSettings::EFFTSize::Min_64;
-		break;
-	case EFFTSize::Small:
-		OutSettings.FFTSize = Audio::FSpectrumAnalyzerSettings::EFFTSize::Small_256;
-		break;
-	case EFFTSize::Medium:
-		OutSettings.FFTSize = Audio::FSpectrumAnalyzerSettings::EFFTSize::Medium_512;
-		break;
-	case EFFTSize::Large:
-		OutSettings.FFTSize = Audio::FSpectrumAnalyzerSettings::EFFTSize::Large_1024;
-		break;
-	case EFFTSize::VeryLarge:
-		OutSettings.FFTSize = Audio::FSpectrumAnalyzerSettings::EFFTSize::VeryLarge_2048;
-		break;
-	case EFFTSize::Max:
-		OutSettings.FFTSize = Audio::FSpectrumAnalyzerSettings::EFFTSize::TestLarge_4096;
-		break;
-	default:
-		break;
-	}
-
-	switch (InterpolationMethod)
-	{
-	case EFFTPeakInterpolationMethod::NearestNeighbor:
-		OutSettings.InterpolationMethod = Audio::FSpectrumAnalyzerSettings::EPeakInterpolationMethod::NearestNeighbor;
-		break;
-	case EFFTPeakInterpolationMethod::Linear:
-		OutSettings.InterpolationMethod = Audio::FSpectrumAnalyzerSettings::EPeakInterpolationMethod::Linear;
-		break;
-	case EFFTPeakInterpolationMethod::Quadratic:
-		OutSettings.InterpolationMethod = Audio::FSpectrumAnalyzerSettings::EPeakInterpolationMethod::Quadratic;
-		break;
-	case EFFTPeakInterpolationMethod::ConstantQ:
-		OutSettings.InterpolationMethod = Audio::FSpectrumAnalyzerSettings::EPeakInterpolationMethod::ConstantQ;
-	default:
-		break;
-	}
+	OutSettings.FFTSize = FFTSize;
+	OutSettings.WindowType = WindowType;
+	OutSettings.InterpolationMethod = InterpolationMethod;
+	OutSettings.HopSize = HopSize;
+	OutSettings.SpectrumType = SpectrumType;
 
-	switch (WindowType)
-	{
-	case EFFTWindowType::None:
-		OutSettings.WindowType = Audio::EWindowType::None;
-		break;
-	case EFFTWindowType::Hamming:
-		OutSettings.WindowType = Audio::EWindowType::Hamming;
-		break;
-	case EFFTWindowType::Hann:
-		OutSettings.WindowType = Audio::EWindowType::Hann;
-		break;
-	case EFFTWindowType::Blackman:
-		OutSettings.WindowType = Audio::EWindowType::Blackman;
-		break;
-	default:
-		break;
-	}
+	return OutSettings;
+}
 
-	switch (AudioSpectrumType)
-	{
-	case EAudioSpectrumType::MagnitudeSpectrum:
-		OutSettings.SpectrumType = Audio::FSpectrumAnalyzerSettings::ESpectrumAnalyzerType::Magnitude;
-		break;
-	case EAudioSpectrumType::PowerSpectrum:
-		OutSettings.SpectrumType = Audio::FSpectrumAnalyzerSettings::ESpectrumAnalyzerType::Power;
-		break;
-	case EAudioSpectrumType::Decibel:
-		OutSettings.SpectrumType = Audio::FSpectrumAnalyzerSettings::ESpectrumAnalyzerType::Decibel;
-		break;
-	}
+FSoundSpectrumAnalyzerDelegateSettings USoundSubmix::GetSpectrumAnalysisDelegateSettings(const TArray<FSoundSubmixSpectralAnalysisBandSettings>& InBandSettings, float UpdateRate, float DecibelNoiseFloor, bool bDoNormalize, bool bDoAutoRange, float AutoRangeAttackTime, float AutoRangeReleaseTime)
+{
+	FSoundSpectrumAnalyzerDelegateSettings DelegateSettings;
 
-	OutSettings.HopSize = HopSize;
+	DelegateSettings.BandSettings = InBandSettings;
+	DelegateSettings.UpdateRate = UpdateRate;
+	DelegateSettings.DecibelNoiseFloor = DecibelNoiseFloor;
+	DelegateSettings.bDoNormalize = bDoNormalize;
+	DelegateSettings.bDoAutoRange = bDoAutoRange;
+	DelegateSettings.AutoRangeAttackTime = AutoRangeAttackTime;
+	DelegateSettings.AutoRangeReleaseTime = AutoRangeReleaseTime;
 
-	return OutSettings;
+	return DelegateSettings;
 }
 
 void USoundfieldEndpointSubmix::SanitizeLinks()
diff --git a/Engine/Source/Runtime/Engine/Public/AudioDevice.h b/Engine/Source/Runtime/Engine/Public/AudioDevice.h
--- a/Engine/Source/Runtime/Engine/Public/AudioDevice.h
+++ b/Engine/Source/Runtime/Engine/Public/AudioDevice.h
@@ -16,6 +16,7 @@
 #include "Sound/SoundClass.h"
 #include "Sound/SoundConcurrency.h"
 #include "Sound/SoundMix.h"
+#include "Sound/SoundSubmix.h"
 #include "Sound/SoundSubmixSend.h"
 #include "Sound/SoundSourceBus.h"
 #include "AudioVirtualLoop.h"
@@ -1304,11 +1305,12 @@ public:
 
 	/** Adds an envelope follower delegate to the submix for this audio device. */
 	virtual void AddEnvelopeFollowerDelegate(USoundSubmix* InSubmix, const FOnSubmixEnvelopeBP& OnSubmixEnvelopeBP);
-	virtual void StartSpectrumAnalysis(USoundSubmix* InSubmix, const Audio::FSpectrumAnalyzerSettings& InSettings);
+
+	virtual void StartSpectrumAnalysis(USoundSubmix* InSubmix, const FSoundSpectrumAnalyzerSettings& InSettings);
 	virtual void StopSpectrumAnalysis(USoundSubmix* InSubmix);
 	virtual void GetMagnitudesForFrequencies(USoundSubmix* InSubmix, const TArray<float>& InFrequencies, TArray<float>& OutMagnitudes);
 	virtual void GetPhasesForFrequencies(USoundSubmix* InSubmix, const TArray<float>& InFrequencies, TArray<float>& OutPhases);
-	virtual void AddSpectralAnalysisDelegate(USoundSubmix* InSubmix, const TArray<FSoundSubmixSpectralAnalysisBandSettings>& InBandSettings, const FOnSubmixSpectralAnalysisBP& OnSubmixSpectralAnalysisBP, float UpdateRate);
+	virtual void AddSpectralAnalysisDelegate(USoundSubmix* InSubmix, const FSoundSpectrumAnalyzerDelegateSettings& InDelegateSettings, const FOnSubmixSpectralAnalysisBP& OnSubmixSpectralAnalysisBP);
 	virtual void RemoveSpectralAnalysisDelegate(USoundSubmix* InSubmix, const FOnSubmixSpectralAnalysisBP& OnSubmixSpectralAnalysisBP);
 
 
diff --git a/Engine/Source/Runtime/SignalProcessing/Private/FloatArrayMath.cpp b/Engine/Source/Runtime/SignalProcessing/Private/FloatArrayMath.cpp
--- a/Engine/Source/Runtime/SignalProcessing/Private/FloatArrayMath.cpp
+++ b/Engine/Source/Runtime/SignalProcessing/Private/FloatArrayMath.cpp
@@ -523,6 +523,17 @@ namespace Audio
 		}
 	}
 
+	void ArraySqrtInPlace(TArrayView<float> InValues)
+	{
+		const int32 Num = InValues.Num();
+		float* InValuesData = InValues.GetData();
+
+		for (int32 i = 0; i < Num; i++)
+		{
+			InValues[i] = FMath::Sqrt(InValues[i]);
+		}
+	}
+
 	void ArrayMagnitudeToDecibelInPlace(TArrayView<float> InValues, float InMinimumDb)
 	{
 		const int32 Num = InValues.Num();
diff --git a/Engine/Source/Runtime/SignalProcessing/Private/SpectrumAnalyzer.cpp b/Engine/Source/Runtime/SignalProcessing/Private/SpectrumAnalyzer.cpp
--- a/Engine/Source/Runtime/SignalProcessing/Private/SpectrumAnalyzer.cpp
+++ b/Engine/Source/Runtime/SignalProcessing/Private/SpectrumAnalyzer.cpp
@@ -4,6 +4,8 @@
 #include "DSP/FFTAlgorithm.h"
 #include "DSP/ConstantQ.h"
 #include "SignalProcessingModule.h"
+#include "Algo/MinElement.h"
+#include "Algo/MaxElement.h"
 
 namespace Audio
 {
@@ -20,48 +22,34 @@ namespace Audio
 	// Implementation of spectrum band extractor
 	class FSpectrumBandExtractor : public ISpectrumBandExtractor
 	{
+		using FBandSettings = ISpectrumBandExtractor::FBandSettings;
 			// FBandSpec describes specifications for a single band.
-			struct FBandSpec
+			struct FBandSpec : public FBandSettings
 			{
 				// Location in output array where band value should be stored.
 				int32 OutIndex;
 
-				// Center frequency of the band.
-				float CenterFrequency;
-
-				// The metric used for the band value.
-				EMetric Metric;
-
-				// The noisefloor in decibels, used when the metric is decibels. 
-				float DbNoiseFloor;
-
 				// The scaling parameter to apply to the power spectrum.
 				float PowerSpectrumScale;
 
-				// If true, all values are scaled and clamped between 0.0 and 1.0.
-				bool bDoNormalize;
-
-				FBandSpec(const FSpectrumBandExtractorSettings& InSettings, int32 InOutIndex, float InCenterFrequency, EMetric InMetric, float InDecibelNoiseFloor, bool bInDoNormalize)
-				:	OutIndex(InOutIndex)
-				,	CenterFrequency(InCenterFrequency)
-				,	Metric(InMetric)
-				,	DbNoiseFloor(InDecibelNoiseFloor)
+				FBandSpec(const FBandSettings& InBandSettings, const FSpectrumBandExtractorSettings& InSettings, const FSpectrumBandExtractorSpectrumSettings& InSpectrumSettings, int32 InOutIndex)
+				:	FBandSettings(InBandSettings)
+				,	OutIndex(InOutIndex)
 				,	PowerSpectrumScale(1.f)
-				,	bDoNormalize(bInDoNormalize)
 				{
-					Update(InSettings);
+					Update(InSettings, InSpectrumSettings);
 				}
 
 				virtual ~FBandSpec() {}
 
 				// Update calculates parameters that are specific to FFT implementation
 				// and sample rate.
-				virtual void Update(const FSpectrumBandExtractorSettings& InSettings)
+				virtual void Update(const FSpectrumBandExtractorSettings& InSettings, const FSpectrumBandExtractorSpectrumSettings& InSpectrumSettings)
 				{
 					PowerSpectrumScale = 1.f;
-					float FloatFFTSize = FMath::Max(static_cast<float>(InSettings.FFTSize), 1.f);
+					float FloatFFTSize = FMath::Max(static_cast<float>(InSpectrumSettings.FFTSize), 1.f);
 
-					switch (InSettings.FFTScaling)
+					switch (InSpectrumSettings.FFTScaling)
 					{
 						case EFFTScaling::MultipliedByFFTSize:
 							PowerSpectrumScale = 1.f / (FloatFFTSize * FloatFFTSize);
@@ -84,52 +72,18 @@ namespace Audio
 							PowerSpectrumScale = 1.f;
 							break;
 					}
+				}
 
-					/*
-					float WindowScale = 1.f;
-
-
-					switch (InSettings.WindowType)
-					{
-						case EWindowType::None:
-							WindowScale = 1.f / FloatFFTSize;
-							break;
-
-						case EWindowType::Hamming:
-							// General form of scaling for generalized cosine on powers spectrum is
-							// 1 / ((1.5 * alpha^2 - alpha + 0.5) * FFTSize)
-
-							// For hamming, alpha = 0.54
-							WindowScale = 1.f / (0.3974 * FloatFFTSize);
-							break;
-
-						case EWindowType::Hann:
-							// General form of scaling for generalized cosine on powers spectrum is
-							// 1 / ((1.5 * alpha^2 - alpha + 0.5) * FFTSize)
-
-							// For hann, alpha = 0.5
-							WindowScale = 1.f / (0.375f * FloatFFTSize);
-							break;
-
-						case EWindowType::Blackman:
-							// General form of scaling for blackman windows is
-							// 1 / ((alpha_0^2 + alpha_1^2 / 2 + alpha_2^2 / 2) * FFTSize
-							//
-							// For this window alpha_0 = 0.42, alpha_1 = 0.5 and alph_2 = 0.08
-							WindowScale = 1.f / (0.3046f * FloatFFTSize);
-							break;
-
-						default:
-							// Should not get here. Means that a window type is not covered.
-							UE_LOG(LogSignalProcessing, Warning, TEXT("Invalid window type enum encountered, %d"), (int32)InSettings.WindowType);
-							WindowScale = 1.f / FloatFFTSize;
-							break;
-
-					}
+				
+				void UpdateOutputFromPowerSpectrum(const float* InputBuffer, int32 InNum, float* OutputBuffer, int32 OutNum) const
+				{
+					check(OutIndex >= 0);
+					check(OutIndex < OutNum);
 
-					PowerSpectrumScale *= WindowScale;
-					*/
+					OutputBuffer[OutIndex] = ExtractFromPowerSpectrum(InputBuffer, InNum) * PowerSpectrumScale;
 				}
+
+				virtual float ExtractFromPowerSpectrum(const float* InputBuffer, int32 InNum) const = 0;
 			};
 
 			// Specification for a nearest neighbor band.
@@ -141,17 +95,24 @@ namespace Audio
 				// Index in power spectrum to lookup band.
 				int32 Index;
 
-				virtual void Update(const FSpectrumBandExtractorSettings& InSettings) override
+				virtual void Update(const FSpectrumBandExtractorSettings& InSettings, const FSpectrumBandExtractorSpectrumSettings& InSpectrumSettings) override
 				{
 					// Call parent class update.
-					FBandSpec::Update(InSettings);
+					FBandSpec::Update(InSettings, InSpectrumSettings);
 
 					// Update the index
-					const int32 MaxSpectrumIndex = InSettings.FFTSize / 2;
-					const float Position = CenterFrequency / FMath::Max(InSettings.SampleRate, 1.f) * InSettings.FFTSize;
+					const int32 MaxSpectrumIndex = InSpectrumSettings.FFTSize / 2;
+					const float Position = CenterFrequency / FMath::Max(InSpectrumSettings.SampleRate, 1.f) * InSpectrumSettings.FFTSize;
 					Index = FMath::RoundToInt(Position);
 					Index = FMath::Clamp(Index, 0, MaxSpectrumIndex);
 				}
+
+				virtual float ExtractFromPowerSpectrum(const float* InputBuffer, int32 InNum) const override
+				{
+					int32 SafeIndex = FMath::Clamp(Index, 0, InNum);
+
+					return InputBuffer[SafeIndex];
+				}
 			};
 
 			// Specification for a linearly interpolated band.
@@ -169,14 +130,14 @@ namespace Audio
 				// Value used for lerping between lower and upper band values.
 				float Alpha;
 
-				virtual void Update(const FSpectrumBandExtractorSettings& InSettings) override
+				virtual void Update(const FSpectrumBandExtractorSettings& InSettings, const FSpectrumBandExtractorSpectrumSettings& InSpectrumSettings) override
 				{
 					// Call parent class update.
-					FBandSpec::Update(InSettings);
+					FBandSpec::Update(InSettings, InSpectrumSettings);
 
 					// Update lower index, upper index and alpha.
-					const int32 MaxSpectrumIndex = InSettings.FFTSize / 2;
-					const float Position = CenterFrequency / FMath::Max(InSettings.SampleRate, 1.f) * InSettings.FFTSize;
+					const int32 MaxSpectrumIndex = InSpectrumSettings.FFTSize / 2;
+					const float Position = CenterFrequency / FMath::Max(InSpectrumSettings.SampleRate, 1.f) * InSpectrumSettings.FFTSize;
 
 					LowerIndex = FMath::FloorToInt(Position);
 					UpperIndex = LowerIndex + 1;
@@ -186,6 +147,14 @@ namespace Audio
 					UpperIndex = FMath::Clamp(UpperIndex, 0, MaxSpectrumIndex);
 					Alpha = FMath::Clamp(Alpha, 0.f, 1.f);
 				}
+
+				virtual float ExtractFromPowerSpectrum(const float* InputBuffer, int32 InNum) const override
+				{
+					int32 SafeLowerIndex = FMath::Clamp(LowerIndex, 0, InNum);
+					int32 SafeUpperIndex = FMath::Clamp(UpperIndex, 0, InNum);
+
+					return FMath::Lerp<float>(InputBuffer[SafeLowerIndex], InputBuffer[SafeUpperIndex], Alpha);
+				}
 			};
 
 			// Specification for band using quadratic interpolation.
@@ -212,14 +181,16 @@ namespace Audio
 				// Weight for upper value.
 				float UpperWeight;
 
-				virtual void Update(const FSpectrumBandExtractorSettings& InSettings) override
+				virtual void Update(const FSpectrumBandExtractorSettings& InSettings, const FSpectrumBandExtractorSpectrumSettings& InSpectrumSettings) override
 				{
 					// Call parent class update.
-					FBandSpec::Update(InSettings);
+					FBandSpec::Update(InSettings, InSpectrumSettings);
+
+					QFactor = FMath::Clamp(QFactor, 0.0001f, 10000.f);
 
 					// Update indices and weights.
-					const int32 MaxSpectrumIndex = InSettings.FFTSize / 2;
-					const float Position = CenterFrequency / FMath::Max(InSettings.SampleRate, 1.f) * InSettings.FFTSize;
+					const int32 MaxSpectrumIndex = InSpectrumSettings.FFTSize / 2;
+					const float Position = CenterFrequency / FMath::Max(InSpectrumSettings.SampleRate, 1.f) * InSpectrumSettings.FFTSize;
 
 					MidIndex = FMath::RoundToInt(Position);
 					LowerIndex = MidIndex - 1;
@@ -236,6 +207,19 @@ namespace Audio
 					MidIndex = FMath::Clamp(MidIndex, 0, MaxSpectrumIndex);
 					UpperIndex = FMath::Clamp(UpperIndex, 0, MaxSpectrumIndex);
 				}
+
+				virtual float ExtractFromPowerSpectrum(const float* InputBuffer, int32 InNum) const override
+				{
+					int32 SafeLowerIndex = FMath::Clamp(LowerIndex, 0, InNum);
+					int32 SafeMidIndex = FMath::Clamp(MidIndex, 0, InNum);
+					int32 SafeUpperIndex = FMath::Clamp(UpperIndex, 0, InNum);
+
+					const float LowerValue = InputBuffer[SafeLowerIndex];
+					const float MidValue = InputBuffer[SafeMidIndex];
+					const float UpperValue = InputBuffer[SafeUpperIndex];
+					
+					return (LowerValue * LowerWeight) + (MidValue * MidWeight) + (UpperValue * UpperWeight);
+				}
 			};
 
 			// Specification for band using CQT band.
@@ -244,9 +228,6 @@ namespace Audio
 				// Use parent constructor
 				using FBandSpec::FBandSpec;
 
-				// QFactor controls the band width.
-				float QFactor;
-				
 				// Start index in power spectrum 
 				int32 StartIndex;
 
@@ -256,19 +237,19 @@ namespace Audio
 				// Internal buffer used when calculating band.
 				mutable AlignedFloatBuffer WorkBuffer;
 
-				virtual void Update(const FSpectrumBandExtractorSettings& InSettings) override
+				virtual void Update(const FSpectrumBandExtractorSettings& InSettings, const FSpectrumBandExtractorSpectrumSettings& InSpectrumSettings) override
 				{
 					// Call parent class update.
-					FBandSpec::Update(InSettings);
+					FBandSpec::Update(InSettings, InSpectrumSettings);
 
 					// Update band weights and offset index.
-					const int32 MaxSpectrumIndex = InSettings.FFTSize / 2;
-					const float Position = CenterFrequency / FMath::Max(InSettings.SampleRate, 1.f) * InSettings.FFTSize;
+					const int32 MaxSpectrumIndex = InSpectrumSettings.FFTSize / 2;
+					const float Position = CenterFrequency / FMath::Max(InSpectrumSettings.SampleRate, 1.f) * InSpectrumSettings.FFTSize;
 					FPseudoConstantQBandSettings CQTBandSettings;
 					CQTBandSettings.CenterFreq = CenterFrequency;
 					CQTBandSettings.BandWidth = FMath::Max(SMALL_NUMBER, CenterFrequency / FMath::Max(SMALL_NUMBER, QFactor));
-					CQTBandSettings.FFTSize = InSettings.FFTSize;
-					CQTBandSettings.SampleRate = FMath::Max(1.f, InSettings.SampleRate);
+					CQTBandSettings.FFTSize = InSpectrumSettings.FFTSize;
+					CQTBandSettings.SampleRate = FMath::Max(1.f, InSpectrumSettings.SampleRate);
 					CQTBandSettings.Normalization = EPseudoConstantQNormalization::EqualEnergy;
 
 					StartIndex = 0;
@@ -282,21 +263,156 @@ namespace Audio
 						WorkBuffer.AddUninitialized(Weights.Num());
 					}
 				}
+
+				virtual float ExtractFromPowerSpectrum(const float* InputBuffer, int32 InNum) const override
+				{
+					int32 SafeStartIndex = FMath::Clamp(StartIndex, 0, InNum);
+
+					if (ensure((SafeStartIndex + Weights.Num()) <= InNum))
+					{
+
+						float Value = 0.f;
+						int32 NumWeights = Weights.Num();
+
+						if (NumWeights > 0)
+						{
+							check(NumWeights == WorkBuffer.Num());
+							FMemory::Memcpy(WorkBuffer.GetData(), &InputBuffer[SafeStartIndex], NumWeights * sizeof(float));
+
+							ArrayMultiplyInPlace(Weights, WorkBuffer);
+							ArraySum(WorkBuffer, Value);
+						}
+
+						return Value;
+					}
+
+					return 0.f;
+				}
+			};
+
+			// Tracks minimum/maximum values given an attack/release time.
+			class FAutoRange 
+			{
+				float CurrentMinimum;
+				float CurrentMaximum;
+				float AttackTimeConstant;
+				float ReleaseTimeConstant;
+
+				bool bAreCurrentValuesValid;
+
+				public:
+
+					FAutoRange(float InAttackTime = 0.5f, float InReleaseTime = 30.f)
+					:	CurrentMinimum(0.f)
+					,	CurrentMaximum(0.f)
+					,	AttackTimeConstant(1.f)
+					,	ReleaseTimeConstant(1.f)
+					,	bAreCurrentValuesValid(false)
+					{
+						SetAttackTime(InAttackTime);
+						SetReleaseTime(InReleaseTime);
+					}	
+
+					void Set(float InMinimum, float InMaximum)
+					{
+						CurrentMinimum = InMinimum;
+						CurrentMaximum = InMaximum;
+						bAreCurrentValuesValid = true;
+					}
+
+					
+					void Update(float InMinimumValue, float InMaximumValue, float InTimeDelta)
+					{
+						if (bAreCurrentValuesValid)
+						{
+							InTimeDelta = FMath::Max(0.f, InTimeDelta);
+
+							// The time delta isn't constant between updates, so we need to determine
+							// how much the value will change.
+							const float AttackCoef = InTimeDelta * AttackTimeConstant;
+							const float ReleaseCoef = InTimeDelta * ReleaseTimeConstant;
+
+							const float MaxDiff = InMaximumValue - CurrentMaximum;
+
+							if (MaxDiff > 0.f)
+							{
+								CurrentMaximum += MaxDiff * AttackCoef;
+							}
+							else
+							{
+								CurrentMaximum += MaxDiff * ReleaseCoef;
+							}
+
+							const float MinDiff = InMinimumValue - CurrentMinimum;
+
+							if (MinDiff < 0.f)
+							{
+								CurrentMinimum += MinDiff * AttackCoef;
+							}
+							else
+							{
+								CurrentMinimum += MinDiff * ReleaseCoef;
+							}
+
+						}
+						else
+						{
+							// On first time around, initialize to given min/max values.
+							Set(InMinimumValue, InMaximumValue);
+
+						}
+					}
+
+					// Normalize input values to the min/max range. 
+					void Normalize(TArrayView<float> InValues) const
+					{
+						
+						ArrayClampInPlace(InValues, CurrentMinimum, CurrentMaximum);
+
+						ArraySubtractByConstantInPlace(InValues, CurrentMinimum);
+
+						float Range = FMath::Max(SMALL_NUMBER, CurrentMaximum - CurrentMinimum);
+
+						ArrayMultiplyByConstantInPlace(InValues, 1.f / Range);
+					}
+
+					void SetAttackTime(float InAttackTime)
+					{
+						// Factor set to achieve 90% of value by attack time.
+						AttackTimeConstant = 0.9f / FMath::Max(0.001f, InAttackTime);
+					}
+					
+					void SetReleaseTime(float InReleaseTime)
+					{
+						// Factor set to achieve 90% of value by release time.
+						ReleaseTimeConstant = 0.9f / FMath::Max(0.001f, InReleaseTime);
+					}
+
 			};
 
 
 		public:
 			FSpectrumBandExtractor(const FSpectrumBandExtractorSettings& InSettings)
 			:	Settings(InSettings)
+			,	LastTimestamp(0)
 			{
+				AutoRange.SetAttackTime(Settings.AutoRangeAttackTimeInSeconds);
+				AutoRange.SetReleaseTime(Settings.AutoRangeReleaseTimeInSeconds);
 			}
 
 			virtual void SetSettings(const FSpectrumBandExtractorSettings& InSettings) override
 			{
-				bool bSettingsChanged = (Settings != InSettings);
 				Settings = InSettings;
+				AutoRange.SetAttackTime(Settings.AutoRangeAttackTimeInSeconds);
+				AutoRange.SetReleaseTime(Settings.AutoRangeReleaseTimeInSeconds);
+			}
+
+			virtual void SetSpectrumSettings(const FSpectrumBandExtractorSpectrumSettings& InSpectrumSettings) override
+			{
+				bool bSpectrumSettingsChanged = (SpectrumSettings != InSpectrumSettings);
+				SpectrumSettings = InSpectrumSettings;
 
-				if (bSettingsChanged)
+				if (bSpectrumSettingsChanged)
 				{
 
 					// If the settings have changed from the previous call, the band specs
@@ -325,37 +441,39 @@ namespace Audio
 				return Num;
 			}
 
-			// Add a nearest neighbor band
-			virtual void AddNearestNeighborBand(float InCenterFrequency, EMetric InMetric, float InDecibelNoiseFloor, bool bInDoNormalize) override
+			virtual void AddBand(const FBandSettings& InBandSettings) override
 			{
-				AddBand<FNNBandSpec>(NNBandSpecs, InCenterFrequency, InMetric, InDecibelNoiseFloor, bInDoNormalize);
-			}
+				switch (InBandSettings.Type)
+				{
+					case EBandType::NearestNeighbor:
+						AddBand<FNNBandSpec>(NNBandSpecs, InBandSettings);
+						break;
 
-			// Add a linear interpolated band.
-			virtual void AddLerpBand(float InCenterFrequency, EMetric InMetric, float InDecibelNoiseFloor, bool bInDoNormalize) override
-			{
-				AddBand<FLerpBandSpec>(LerpBandSpecs, InCenterFrequency, InMetric, InDecibelNoiseFloor, bInDoNormalize);
-			}
+					case EBandType::Lerp:
+						AddBand<FLerpBandSpec>(LerpBandSpecs, InBandSettings);
+						break;
 
-			// Add a quadratic interplated band.
-			virtual void AddQuadraticBand(float InCenterFrequency, EMetric InMetric, float InDecibelNoiseFloor, bool bInDoNormalize) override
-			{
-				AddBand<FQuadraticBandSpec>(QuadraticBandSpecs, InCenterFrequency, InMetric, InDecibelNoiseFloor, bInDoNormalize);
-			}
+					case EBandType::Quadratic:
+						AddBand<FQuadraticBandSpec>(QuadraticBandSpecs, InBandSettings);
+						break;
 
-			// Add a constant Q band
-			virtual void AddConstantQBand(float InCenterFrequency, float InQFactor, EMetric InMetric, float InDecibelNoiseFloor, bool bInDoNormalize) override
-			{
-				FCQTBandSpec& Band = AddBand<FCQTBandSpec>(CQTBandSpecs, InCenterFrequency, InMetric, InDecibelNoiseFloor, bInDoNormalize);
-				Band.QFactor = FMath::Clamp(InQFactor, 0.001f, 100.f);
+					case EBandType::ConstantQ:
+					default:
+						AddBand<FCQTBandSpec>(CQTBandSpecs, InBandSettings);
+						break;
+				}
 			}
 
+
 			// Extract band from input.
-			virtual void ExtractBands(const AlignedFloatBuffer& InComplexBuffer, TArray<float>& OutValues) override
+			virtual void ExtractBands(const AlignedFloatBuffer& InComplexBuffer, double InTimestamp, TArray<float>& OutValues) override
 			{
 				const int32 NumComplex = InComplexBuffer.Num();
 
-				check(NumComplex == (Settings.FFTSize + 2));
+				float TimeDelta = FMath::Max(0.0f, static_cast<float>(InTimestamp - LastTimestamp));
+				LastTimestamp = InTimestamp;
+
+				check(NumComplex == (SpectrumSettings.FFTSize + 2));
 
 				OutValues.Reset();
 				OutValues.AddZeroed(GetNumBands());
@@ -373,17 +491,40 @@ namespace Audio
 				ExtractBands(PowerSpectrum, LerpBandSpecs, OutValues);
 				ExtractBands(PowerSpectrum, QuadraticBandSpecs, OutValues);
 				ExtractBands(PowerSpectrum, CQTBandSpecs, OutValues);
+
+				ApplyMetric(OutValues);
+
+				if (Settings.bDoAutoRange)
+				{
+					Normalize(OutValues);
+
+					float* MinimumValuePtr = Algo::MinElement(OutValues);
+					float* MaximumValuePtr = Algo::MaxElement(OutValues);
+
+					float MinimumValue = MinimumValuePtr == nullptr ? 0.f : *MinimumValuePtr;
+					float MaximumValue = MaximumValuePtr == nullptr ? 0.f : *MaximumValuePtr;
+
+					AutoRange.Update(MinimumValue, MaximumValue, TimeDelta);
+
+					AutoRange.Normalize(OutValues);
+				}
+				else if (Settings.bDoNormalize)
+				{
+					Normalize(OutValues);
+
+					ArrayClampInPlace(OutValues, 0.f, 1.f);
+				}
 			}
 
 		private:
 
 			// Adds a band spec and returns a reference to the added spec.
 			template<typename T> 
-			T& AddBand(TArray<T>& InBandSpecs, float InCenterFrequency, EMetric InMetric, float InDecibelNoiseFloor, bool bInDoNormalize)
+			T& AddBand(TArray<T>& InBandSpecs, const FBandSettings& InBandSettings)
 			{
 				int32 OutIndex = GetNumBands();
 
-				T BandSpec(Settings, OutIndex, InCenterFrequency, InMetric, InDecibelNoiseFloor, bInDoNormalize);
+				T BandSpec(InBandSettings, Settings, SpectrumSettings, OutIndex);
 
 				return InBandSpecs.Add_GetRef(BandSpec);
 			}
@@ -394,7 +535,7 @@ namespace Audio
 			{
 				for (T& BandSpec : InSpecs)
 				{
-					BandSpec.Update(Settings);
+					BandSpec.Update(Settings, SpectrumSettings);
 				}
 			}
 
@@ -407,153 +548,75 @@ namespace Audio
 				UpdateBandSpecs<FCQTBandSpec>(CQTBandSpecs);
 			}
 
-
-			// Extract nearest neighbor bands
-			void ExtractBands(const AlignedFloatBuffer& InPowerSpectrum, const TArray<FNNBandSpec>& InNNBandSpecs, TArray<float>& OutValues) const
-			{
-				float* OutData = OutValues.GetData();
-				const float* InData = InPowerSpectrum.GetData();
-				int32 InNum = InPowerSpectrum.Num();
-
-				for (const FNNBandSpec& Spec : InNNBandSpecs)
-				{
-					check(Spec.OutIndex >= 0);
-					check(Spec.OutIndex < OutValues.Num());
-					check(Spec.Index < InNum);
-					check(Spec.Index >= 0);
-					OutData[Spec.OutIndex] = ApplyScaleAndMetric(Spec, InData[Spec.Index]);
-				};
-			}
-
-			// Extract linearly interpolated bands
-			void ExtractBands(const AlignedFloatBuffer& InPowerSpectrum, const TArray<FLerpBandSpec>& InLerpBandSpecs, TArray<float>& OutValues) const
-			{
-				float* OutData = OutValues.GetData();
-				const float* InData = InPowerSpectrum.GetData();
-				int32 InNum = InPowerSpectrum.Num();
-
-				for (const FLerpBandSpec& Spec : InLerpBandSpecs)
-				{
-					check(Spec.OutIndex >= 0);
-					check(Spec.OutIndex < OutValues.Num());
-					check(Spec.LowerIndex < InNum);
-					check(Spec.LowerIndex >= 0);
-					check(Spec.UpperIndex < InNum);
-					check(Spec.UpperIndex >= 0);
-
-					float Value = FMath::Lerp<float>(InData[Spec.LowerIndex], InData[Spec.UpperIndex], Spec.Alpha);
-					OutData[Spec.OutIndex] = ApplyScaleAndMetric(Spec, Value);
-				};
-			}
-
-			// Extract quadratically interpolated bands.
-			void ExtractBands(const AlignedFloatBuffer& InPowerSpectrum, const TArray<FQuadraticBandSpec>& InQuadraticBandSpecs, TArray<float>& OutValues) const
+			template<typename T>
+			void ExtractBands(const AlignedFloatBuffer& InPowerSpectrum, const TArray<T>& InBandSpecs, TArray<float>& OutValues) const
 			{
 				float* OutData = OutValues.GetData();
+				int32 OutNum = OutValues.Num();
 				const float* InData = InPowerSpectrum.GetData();
 				int32 InNum = InPowerSpectrum.Num();
 
-				for (const FQuadraticBandSpec& Spec : InQuadraticBandSpecs)
+				for (const T& Spec : InBandSpecs)
 				{
-					check(Spec.OutIndex >= 0);
-					check(Spec.OutIndex < OutValues.Num());
-					check(Spec.LowerIndex < InNum);
-					check(Spec.LowerIndex >= 0);
-					check(Spec.MidIndex < InNum);
-					check(Spec.MidIndex >= 0);
-					check(Spec.UpperIndex < InNum);
-					check(Spec.UpperIndex >= 0);
-
-					const float LowerValue = InData[Spec.LowerIndex];
-					const float MidValue = InData[Spec.MidIndex];
-					const float UpperValue = InData[Spec.UpperIndex];
-					
-					float Value = (LowerValue * Spec.LowerWeight) + (MidValue * Spec.MidWeight) + (UpperValue * Spec.UpperWeight);
-
-					OutData[Spec.OutIndex] = ApplyScaleAndMetric(Spec, Value);
+					Spec.UpdateOutputFromPowerSpectrum(InData, InNum, OutData, OutNum);
 				}
 			}
 
-			// Extract constant q bands.
-			void ExtractBands(const AlignedFloatBuffer& InPowerSpectrum, const TArray<FCQTBandSpec>& InCQTBandSpecs, TArray<float>& OutValues) const
+			// apply metric to a band value
+			void ApplyMetric(TArray<float>& InValues) const
 			{
-				float* OutData = OutValues.GetData();
-				const float* InData = InPowerSpectrum.GetData();
-				int32 InNum = InPowerSpectrum.Num();
 
-				for (const FCQTBandSpec& Spec : InCQTBandSpecs)
+				switch (Settings.Metric)
 				{
-					check(Spec.OutIndex >= 0);
-					check(Spec.OutIndex < OutValues.Num());
-					check(Spec.StartIndex < InNum);
-					check(Spec.StartIndex >= 0);
-					check((Spec.StartIndex + Spec.Weights.Num()) <= InNum);
-
-					float Value = 0.f;
-					int32 NumWeights = Spec.Weights.Num();
-
-					if (NumWeights > 0)
-					{
-						check(NumWeights == Spec.WorkBuffer.Num());
-						FMemory::Memcpy(Spec.WorkBuffer.GetData(), &InData[Spec.StartIndex], NumWeights * sizeof(float));
+					case FSpectrumBandExtractorSettings::EMetric::Magnitude:
+						ArraySqrtInPlace(InValues);
+						break;
 
-						ArrayMultiplyInPlace(Spec.Weights, Spec.WorkBuffer);
-						ArraySum(Spec.WorkBuffer, Value);
-					}
+					case FSpectrumBandExtractorSettings::EMetric::Decibel:
+						ArrayPowerToDecibelInPlace(InValues, Settings.DecibelNoiseFloor);
+						break;
 
-					OutData[Spec.OutIndex] = ApplyScaleAndMetric(Spec, Value);
+					case FSpectrumBandExtractorSettings::EMetric::Power:
+					default:
+						break;
 				}
-				
 			}
 
-			// Scale and apply metric to a band value
-			float ApplyScaleAndMetric(const FBandSpec& InBandSpec, float InValue) const
+			void Normalize(TArray<float>& InValues) const
 			{
-				float OutValue = InValue * InBandSpec.PowerSpectrumScale;
-
-				switch (InBandSpec.Metric)
+				// Only decibel values need to be normalized. 
+				if (FSpectrumBandExtractorSettings::EMetric::Decibel == Settings.Metric)
 				{
-					case ISpectrumBandExtractor::EMetric::Magnitude:
-						OutValue = FMath::Sqrt(OutValue);
-						break;
-
-					case ISpectrumBandExtractor::EMetric::Decibel:
-						OutValue = 10.f * FMath::Loge(OutValue) / SpectrumAnalyzerIntrinsics::Loge10;
-						if (!FMath::IsFinite(OutValue) || (OutValue < InBandSpec.DbNoiseFloor))
+					for (float& Value : InValues)
+					{
+						if (!FMath::IsFinite(Value))
 						{
-							OutValue = InBandSpec.DbNoiseFloor;
+							Value = Settings.DecibelNoiseFloor;
 						}
+					}
 
-						if (InBandSpec.bDoNormalize)
-						{
-							OutValue -= InBandSpec.DbNoiseFloor;
-							if (InBandSpec.DbNoiseFloor < 0.f)
-							{
-								OutValue /= (-InBandSpec.DbNoiseFloor);
-							}
-						}
-						break;
+					ArrayClampMinInPlace(InValues, Settings.DecibelNoiseFloor);
 
-					case ISpectrumBandExtractor::EMetric::Power:
-					default:
-						break;
-				}
+					ArraySubtractByConstantInPlace(InValues, Settings.DecibelNoiseFloor);
 
-				if (InBandSpec.bDoNormalize)
-				{
-					OutValue = FMath::Clamp(OutValue, 0.f, 1.f);
-				}
+					float Range = FMath::Max(0.01f, -Settings.DecibelNoiseFloor);
 
-				return OutValue;
+					ArrayMultiplyByConstantInPlace(InValues, 1.f / Range);
+				}
 			}
 
 			FSpectrumBandExtractorSettings Settings;
+			FSpectrumBandExtractorSpectrumSettings SpectrumSettings;
+			double LastTimestamp;
 
 			AlignedFloatBuffer PowerSpectrum;
+
 			TArray<FNNBandSpec> NNBandSpecs;
 			TArray<FLerpBandSpec> LerpBandSpecs;
 			TArray<FQuadraticBandSpec> QuadraticBandSpecs;
 			TArray<FCQTBandSpec> CQTBandSpecs;
+
+			FAutoRange AutoRange;
 	};
 
 
@@ -569,8 +632,10 @@ namespace Audio
 		, bIsInitialized(false)
 		, SampleRate(0.0f)
 		, Window(CurrentSettings.WindowType, (int32)CurrentSettings.FFTSize, 1, false)
+		, SampleCounter(0)
 		, InputQueue(FMath::Max((int32)CurrentSettings.FFTSize * 4, 4096))
 		, FrequencyBuffer((int32)CurrentSettings.FFTSize)
+		, LockedBufferTimestamp(0)
 		, LockedFrequencyVector(nullptr)
 	{
 	}
@@ -581,8 +646,10 @@ namespace Audio
 		, bIsInitialized(true)
 		, SampleRate(InSampleRate)
 		, Window(InSettings.WindowType, (int32)InSettings.FFTSize, 1, false)
+		, SampleCounter(0)
 		, InputQueue(FMath::Max((int32)CurrentSettings.FFTSize * 4, 4096))
 		, FrequencyBuffer((int32)InSettings.FFTSize)
+		, LockedBufferTimestamp(0)
 		, LockedFrequencyVector(nullptr)
 	{
 		ResetSettings();
@@ -594,8 +661,10 @@ namespace Audio
 		, bIsInitialized(true)
 		, SampleRate(InSampleRate)
 		, Window(CurrentSettings.WindowType, (int32)CurrentSettings.FFTSize, 1, false)
+		, SampleCounter(0)
 		, InputQueue(FMath::Max((int32)CurrentSettings.FFTSize * 4, 4096))
 		, FrequencyBuffer((int32)CurrentSettings.FFTSize)
+		, LockedBufferTimestamp(0)
 		, LockedFrequencyVector(nullptr)
 	{
 		ResetSettings();
@@ -620,6 +689,7 @@ namespace Audio
 		CurrentSettings = InSettings;
 		bSettingsWereUpdated = false;
 		SampleRate = InSampleRate;
+		SampleCounter.Set(0);
 		InputQueue.SetCapacity(FMath::Max((int32)CurrentSettings.FFTSize * 4, 4096));
 		FrequencyBuffer.Reset((int32)CurrentSettings.FFTSize);
 		ResetSettings();
@@ -700,7 +770,7 @@ namespace Audio
 		bSettingsWereUpdated = false;
 	}
 
-	void FSpectrumAnalyzer::PerformInterpolation(const AlignedFloatBuffer& InComplexBuffer, FSpectrumAnalyzerSettings::EPeakInterpolationMethod InMethod, const float InFreq, float& OutReal, float& OutImag)
+	void FSpectrumAnalyzer::PerformInterpolation(const AlignedFloatBuffer& InComplexBuffer, FSpectrumAnalyzer::EPeakInterpolationMethod InMethod, const float InFreq, float& OutReal, float& OutImag)
 	{
 		const float* InComplexData = InComplexBuffer.GetData();
 		const int32 VectorLength = InComplexBuffer.Num();
@@ -715,7 +785,7 @@ namespace Audio
 		
 		switch (InMethod)
 		{
-		case Audio::FSpectrumAnalyzerSettings::EPeakInterpolationMethod::NearestNeighbor:
+		case Audio::FSpectrumAnalyzer::EPeakInterpolationMethod::NearestNeighbor:
 		{
 			int32 Index = FMath::RoundToInt(Position) & SpectrumAnalyzerIntrinsics::EvenNumberMask;
 
@@ -727,7 +797,7 @@ namespace Audio
 			break;
 		}
 			
-		case Audio::FSpectrumAnalyzerSettings::EPeakInterpolationMethod::Linear:
+		case Audio::FSpectrumAnalyzer::EPeakInterpolationMethod::Linear:
 		{
 			int32 LowerIndex = FMath::FloorToInt(Position) & SpectrumAnalyzerIntrinsics::EvenNumberMask;
 			int32 UpperIndex = LowerIndex + 2;
@@ -748,7 +818,7 @@ namespace Audio
 			OutImag = FMath::Lerp<float>(y1Imag, y2Imag, PositionFraction);
 			break;
 		}	
-		case Audio::FSpectrumAnalyzerSettings::EPeakInterpolationMethod::Quadratic:
+		case Audio::FSpectrumAnalyzer::EPeakInterpolationMethod::Quadratic:
 		{
 			// Note: math here does not interpolate quadratically. 
 			const int32 MidIndex = FMath::Clamp(FMath::RoundToInt(Position) & SpectrumAnalyzerIntrinsics::EvenNumberMask, 0, NyquistPosition);
@@ -788,7 +858,7 @@ namespace Audio
 		OutSettings = CurrentSettings;
 	}
 
-	float FSpectrumAnalyzer::GetMagnitudeForFrequency(float InFrequency)
+	float FSpectrumAnalyzer::GetMagnitudeForFrequency(float InFrequency, FSpectrumAnalyzer::EPeakInterpolationMethod InMethod)
 	{
 		if (!bIsInitialized)
 		{
@@ -814,13 +884,14 @@ namespace Audio
 			float OutMagnitude = 0.0f;
 
 			float InterpolatedReal, InterpolatedImag;
-			PerformInterpolation(*OutVector, CurrentSettings.InterpolationMethod, InFrequency, InterpolatedReal, InterpolatedImag);
+			PerformInterpolation(*OutVector, InMethod, InFrequency, InterpolatedReal, InterpolatedImag);
 
 			OutMagnitude = FMath::Sqrt((InterpolatedReal * InterpolatedReal) + (InterpolatedImag * InterpolatedImag));
 
 			if (bShouldUnlockBuffer)
 			{
 				FrequencyBuffer.UnlockBuffer();
+				LockedBufferTimestamp = 0.0;
 			}
 
 			return OutMagnitude;
@@ -830,7 +901,7 @@ namespace Audio
 		return 0.0f;
 	}
 
-	float FSpectrumAnalyzer::GetPhaseForFrequency(float InFrequency)
+	float FSpectrumAnalyzer::GetPhaseForFrequency(float InFrequency, FSpectrumAnalyzer::EPeakInterpolationMethod InMethod)
 	{
 		if (!bIsInitialized)
 		{
@@ -856,13 +927,14 @@ namespace Audio
 			float OutPhase = 0.0f;
 
 			float InterpolatedReal, InterpolatedImag;
-			PerformInterpolation(*OutVector, CurrentSettings.InterpolationMethod, InFrequency, InterpolatedReal, InterpolatedImag);
+			PerformInterpolation(*OutVector, InMethod, InFrequency, InterpolatedReal, InterpolatedImag);
 
 			OutPhase = FMath::Atan2(InterpolatedImag, InterpolatedReal);
 
 			if (bShouldUnlockBuffer)
 			{
 				FrequencyBuffer.UnlockBuffer();
+				LockedBufferTimestamp = 0.0;
 			}
 
 			return OutPhase;
@@ -882,36 +954,37 @@ namespace Audio
 			return;
 		}
 
-		const AlignedFloatBuffer* OutVector = nullptr;
+		const AlignedFloatBuffer* AnalysisBuffer = nullptr;
 		bool bShouldUnlockBuffer = true;
 
-		FSpectrumBandExtractorSettings ExtractorSettings;
+		FSpectrumBandExtractorSpectrumSettings ExtractorSettings;
 		ExtractorSettings.SampleRate = SampleRate;
 		ExtractorSettings.FFTSize = FFTSize;
 		ExtractorSettings.FFTScaling = FFTScaling;
 		ExtractorSettings.WindowType = Window.GetWindowType();
 
 		// This should have minimal cost if settings have not changed between calls.
-		InExtractor.SetSettings(ExtractorSettings);
+		InExtractor.SetSpectrumSettings(ExtractorSettings);
 
 		if (LockedFrequencyVector)
 		{
-			OutVector = LockedFrequencyVector;
+			AnalysisBuffer = LockedFrequencyVector;
 			bShouldUnlockBuffer = false;
 		}
 		else
 		{
-			OutVector = &FrequencyBuffer.LockMostRecentBuffer();
+			AnalysisBuffer = &FrequencyBuffer.LockMostRecentBuffer(LockedBufferTimestamp);
 		}
 
 		// Perform work.
-		if (OutVector)
+		if (AnalysisBuffer)
 		{
-			InExtractor.ExtractBands(*OutVector, OutValues);
+			InExtractor.ExtractBands(*AnalysisBuffer, LockedBufferTimestamp, OutValues);
 
 			if (bShouldUnlockBuffer)
 			{
 				FrequencyBuffer.UnlockBuffer();
+				LockedBufferTimestamp = 0.0;
 			}
 		}
 	}
@@ -926,9 +999,10 @@ namespace Audio
 		if (LockedFrequencyVector != nullptr)
 		{
 			FrequencyBuffer.UnlockBuffer();
+			LockedBufferTimestamp = 0.0;
 		}
 
-		LockedFrequencyVector = &FrequencyBuffer.LockMostRecentBuffer();
+		LockedFrequencyVector = &FrequencyBuffer.LockMostRecentBuffer(LockedBufferTimestamp);
 	}
 
 	void FSpectrumAnalyzer::UnlockOutputBuffer()
@@ -942,6 +1016,7 @@ namespace Audio
 		{
 			FrequencyBuffer.UnlockBuffer();
 			LockedFrequencyVector = nullptr;
+			LockedBufferTimestamp = 0.0;
 		}
 	}
 
@@ -953,6 +1028,8 @@ namespace Audio
 
 	bool FSpectrumAnalyzer::PushAudio(const float* InBuffer, int32 NumSamples)
 	{
+		SampleCounter.Add(NumSamples);
+
 		return InputQueue.Push(InBuffer, NumSamples) == NumSamples;
 	}
 
@@ -985,12 +1062,15 @@ namespace Audio
 			ResetSettings();
 		}
 
+
 		AlignedFloatBuffer& FFTOutput = FrequencyBuffer.StartWorkOnBuffer();
 
 		// If we have enough audio pushed to the spectrum analyzer and we have an available buffer to work in,
 		// we can start analyzing.
 		if (InputQueue.Num() >= ((uint32)FFTSize))
 		{
+			int64 WindowSampleCenterIndex = 0;
+
 			float* TimeDomainBuffer = AnalysisTimeDomainBuffer.GetData();
 
 			if (bUseLatestAudio)
@@ -998,6 +1078,8 @@ namespace Audio
 				// If we are only using the latest audio, scrap the oldest audio in the InputQueue:
 				InputQueue.SetNum((uint32)FFTSize);
 			}
+			WindowSampleCenterIndex = SampleCounter.GetValue() - InputQueue.Num() + FFTSize / 2;
+			double Timestamp = static_cast<double>(WindowSampleCenterIndex) / FMath::Max(SampleRate, 1.f);
 
 			// Perform pop/peek here based on FFT size and hop amount.
 			const int32 PeekAmount = FFTSize - HopInSamples;
@@ -1024,7 +1106,7 @@ namespace Audio
 			}
 
 			// We're done, so unlock this vector.
-			FrequencyBuffer.StopWorkOnBuffer();
+			FrequencyBuffer.StopWorkOnBuffer(Timestamp);
 
 			return true;
 		}
@@ -1070,6 +1152,9 @@ namespace Audio
 			}
 		}
 
+		Timestamps.Reset();
+		Timestamps.AddZeroed(SpectrumAnalyzerBufferSize);
+
 		InputIndex = 0;
 		OutputIndex = 0;
 	}
@@ -1105,11 +1190,18 @@ namespace Audio
 		return ComplexBuffers[InputIndex];
 	}
 
-	void FSpectrumAnalyzerBuffer::StopWorkOnBuffer()
+	void FSpectrumAnalyzerBuffer::StopWorkOnBuffer(double InTimestamp)
 	{
+		Timestamps[InputIndex] = InTimestamp;
 		IncrementInputIndex();
 	}
 
+	const AlignedFloatBuffer& FSpectrumAnalyzerBuffer::LockMostRecentBuffer(double& OutTimestamp) const
+	{
+		OutTimestamp = Timestamps[OutputIndex];
+		return ComplexBuffers[OutputIndex];
+	}
+
 	const AlignedFloatBuffer& FSpectrumAnalyzerBuffer::LockMostRecentBuffer() const
 	{
 		return ComplexBuffers[OutputIndex];
diff --git a/Engine/Source/Runtime/SignalProcessing/Public/DSP/FloatArrayMath.h b/Engine/Source/Runtime/SignalProcessing/Public/DSP/FloatArrayMath.h
--- a/Engine/Source/Runtime/SignalProcessing/Public/DSP/FloatArrayMath.h
+++ b/Engine/Source/Runtime/SignalProcessing/Public/DSP/FloatArrayMath.h
@@ -113,6 +113,9 @@ namespace Audio
 	SIGNALPROCESSING_API void ArraySubtractByConstantInPlace(TArrayView<float> InValues, float InSubtrahend);
 	SIGNALPROCESSING_API void ArraySubtractByConstantInPlace(AlignedFloatBuffer& InValues, float InSubtrahend);
 
+
+	SIGNALPROCESSING_API void ArraySqrtInPlace(TArrayView<float> InValues);
+
 	/** Convert magnitude values to decibel values in place. db = 20 * log10(val) */
 	SIGNALPROCESSING_API void ArrayMagnitudeToDecibelInPlace(TArrayView<float> InValues, float InMinimumDb);
 	SIGNALPROCESSING_API void ArrayMagnitudeToDecibelInPlace(AlignedFloatBuffer& InValues, float InMinimumDb);
diff --git a/Engine/Source/Runtime/SignalProcessing/Public/DSP/SpectrumAnalyzer.h b/Engine/Source/Runtime/SignalProcessing/Public/DSP/SpectrumAnalyzer.h
--- a/Engine/Source/Runtime/SignalProcessing/Public/DSP/SpectrumAnalyzer.h
+++ b/Engine/Source/Runtime/SignalProcessing/Public/DSP/SpectrumAnalyzer.h
@@ -13,6 +13,7 @@ namespace Audio
 {
 	class IFFTAlgorithm;
 	class FSpectrumAnalyzer;
+
 	struct SIGNALPROCESSING_API FSpectrumAnalyzerSettings
 	{
 		// Actual FFT size used. For FSpectrumAnalyzer, we never zero pad the input buffer.
@@ -28,27 +29,8 @@ namespace Audio
 			TestLarge_4096 = 4096
 		};
 
-		// Peak interpolation method. If the EFFTSize is small but will be densely sampled,
-		// it's worth using a linear or quadratic interpolation method.
-		enum class EPeakInterpolationMethod : uint8
-		{
-			NearestNeighbor,
-			Linear,
-			Quadratic,
-			ConstantQ,
-		};
-
-		enum class ESpectrumAnalyzerType : uint8
-		{
-			Magnitude,
-			Power,
-			Decibel,
-		};
-
 		EWindowType WindowType;
 		EFFTSize FFTSize;
-		EPeakInterpolationMethod InterpolationMethod;
-		ESpectrumAnalyzerType  SpectrumType;
 
 		/**
 			* Hop size as a percentage of FFTSize.
@@ -61,14 +43,58 @@ namespace Audio
 		FSpectrumAnalyzerSettings()
 			: WindowType(EWindowType::Hann)
 			, FFTSize(EFFTSize::Default)
-			, InterpolationMethod(EPeakInterpolationMethod::Linear)
-			, SpectrumType(ESpectrumAnalyzerType::Magnitude)
 			, HopSize(0.0f)
 		{}
 	};
 
+
 	/** Settings for band extractor. */
 	struct SIGNALPROCESSING_API FSpectrumBandExtractorSettings
+	{
+		/** Metric for output band values. */
+		enum class EMetric : uint8
+		{
+			/** Return the magnitude spectrum value. */
+			Magnitude,
+
+			/** Return the power spectrum value. */
+			Power,
+
+			/** Return the decibel spectrum value. Decibels are calculated
+			 * with 0dB equal to 1.f magnitude.  */
+			Decibel
+		};
+
+		/** Metric used to calculate return value. */
+		EMetric Metric = EMetric::Decibel;
+
+
+		/** If the metric is Decibel, this is the minimum decibel value allowed. */
+		float DecibelNoiseFloor = -40.f; 
+
+		/** 
+		 * If true, all values are scaled and clamped between 0.0 and 1.f. In the 
+		 * case of Decibels, 0.0 corresponds to the decibel noise floor and 1.f to 0dB.
+		 * If bDoAutoRange is true, then values are relatively to recent maximum and minimums
+		 * regardless of the metric used.
+		 */
+		bool bDoNormalize = true;
+
+		/** 
+		 * If true and bDoNormalize is true, then values will be scaled between 0 and 1
+		 * based upon relatively recent minimum and maximum values. 
+		 */
+		bool bDoAutoRange = true;
+
+		/** Time in seconds for autorange to reach 99% of a smaller range. */
+		float AutoRangeReleaseTimeInSeconds = 30.f;
+
+		/** Time in seconds for autorange to reach 99% of a larger range. */
+		float AutoRangeAttackTimeInSeconds = 1.f;
+	};
+
+	/** Settings describing the spectrum used for in the band extractor. */
+	struct SIGNALPROCESSING_API FSpectrumBandExtractorSpectrumSettings
 	{
 		/** Sample rate of audio */
 		float SampleRate; 
@@ -82,7 +108,7 @@ namespace Audio
 		/** Window used when perform FFT */
 		EWindowType WindowType;
 
-		FSpectrumBandExtractorSettings()
+		FSpectrumBandExtractorSpectrumSettings()
 		:	SampleRate(48000.f)
 		,	FFTSize(1024)
 		,	FFTScaling(EFFTScaling::None)
@@ -90,7 +116,7 @@ namespace Audio
 		{}
 
 		/** Compare whether two settings structures are equal. */
-		bool operator==(const FSpectrumBandExtractorSettings& Other) const
+		bool operator==(const FSpectrumBandExtractorSpectrumSettings& Other) const
 		{
 			bool bIsEqual = ((SampleRate == Other.SampleRate)
 					&& (FFTSize == Other.FFTSize)
@@ -100,7 +126,7 @@ namespace Audio
 		}
 
 		/** Compare whether two settings structures are not equal. */
-		bool operator!=(const FSpectrumBandExtractorSettings& Other) const
+		bool operator!=(const FSpectrumBandExtractorSpectrumSettings& Other) const
 		{
 			return !(*this == Other);
 		}
@@ -116,24 +142,45 @@ namespace Audio
 	class SIGNALPROCESSING_API ISpectrumBandExtractor
 	{
 		public:
-			/** Metric for output band values. */
-			enum class EMetric : uint8
+
+			enum class EBandType : uint8
 			{
-				/** Return the magnitude spectrum value. */
-				Magnitude,
+				/** Look up band value via nearest FFT band. */
+				NearestNeighbor,
 
-				/** Return the power spectrum value. */
-				Power,
+				/** Band value is calculated by lerping between FFT bands adjacent to center frequency. */
+				Lerp,
 
-				/** Return the decibel spectrum value. Decibels are calculated
-				 * with 0dB equal to 1.f magnitude.  */
-				Decibel
+				/** Band value is calculated by fitting quadratic to 3 adjancent FFT bands and solving for center frequency. */
+				Quadratic,
+
+				/** Band value is calculated by a weighted sum of a window of fft bands around center frequency. The window size is determined by the QFactor of the band. */
+				ConstantQ
+			};
+
+			/** Settings for a single band */
+			struct FBandSettings
+			{
+			 	/** Type of band to extract. */
+				EBandType Type = EBandType::ConstantQ;
+
+			 	/**  Frequency of interest in hz. */
+				float CenterFrequency = 0.f;
+
+			 	/** 
+				 * QFactor is only applicable for the ConstantQ band type. 
+				 * QFactor = CenterFreq / BandWidth. Eg. A small QFactor results in a wide band.
+				 */
+				float QFactor = 10.0f;
 			};
 
 			virtual ~ISpectrumBandExtractor() {}
 
-			/** Set the settings and update cached internal values if needed */
+			/** Sets and updates the settings for the band extractor */
 			virtual void SetSettings(const FSpectrumBandExtractorSettings& InSettings) = 0;
+
+			/** Set the settings and update cached internal values if needed */
+			virtual void SetSpectrumSettings(const FSpectrumBandExtractorSpectrumSettings& InSettings) = 0;
 			
 			/** Removes all added bands. */
 			virtual void RemoveAllBands() = 0;
@@ -141,56 +188,16 @@ namespace Audio
 			/** Returns the total number of bands. */
 			virtual int32 GetNumBands() const = 0;
 
-			/** Adds a band which calculates the band value as the value of the FFT bin nearest to the center frequency.
-			 *
-			 * @param InCenterFrequency - Frequency of interest in hz .
-			 * @param InMetric - Metric used to calculate return value.
-			 * @param InDecibelNoiseFloor - If the metric is Decibel, this is the minimum decibel value allowed.
-			 * @param bInDoNormalize - If true, all values are scaled and clamped between 0.0 and 1.f. In the 
-			 *                         case of Decibels, 0.0 corresponds to the decibel noise floor and 1.f to 0dB.
-			 */
-			virtual void AddNearestNeighborBand(float InCenterFrequency, EMetric InMetric = EMetric::Decibel, float InDecibelNoiseFloor=-40.f, bool bInDoNormalize=true) = 0;
-
-			/** Adds a band which calculates the band value as a linear interpolation of the values of the FFT 
-			 *  bins adjacent to the center frequency.
-			 *
-			 * @param InCenterFrequency - Frequency of interest in hz .
-			 * @param InMetric - Metric used to calculate return value.
-			 * @param InDecibelNoiseFloor - If the metric is Decibel, this is the minimum decibel value allowed.
-			 * @param bInDoNormalize - If true, all values are scaled and clamped between 0.0 and 1.f. In the 
-			 *                         case of Decibels, 0.0 corresponds to the decibel noise floor and 1.f to 0dB.
-			 */
-			virtual void AddLerpBand(float InCenterFrequency, EMetric InMetric = EMetric::Decibel, float InDecibelNoiseFloor=-40.f, bool bInDoNormalize=true) = 0;
-
-			/** Adds a band which calculates the band value as a quadratic interpolation of the values of the FFT 
-			 *  bins adjacent to the center frequency.
-			 *
-			 * @param InCenterFrequency - Frequency of interest in hz .
-			 * @param InMetric - Metric used to calculate return value.
-			 * @param InDecibelNoiseFloor - If the metric is Decibel, this is the minimum decibel value allowed.
-			 * @param bInDoNormalize - If true, all values are scaled and clamped between 0.0 and 1.f. In the 
-			 *                         case of Decibels, 0.0 corresponds to the decibel noise floor and 1.f to 0dB.
-			 */
-			virtual void AddQuadraticBand(float InCenterFrequency, EMetric InMetric = EMetric::Decibel, float InDecibelNoiseFloor=-40.f, bool bInDoNormalize=true) = 0;
-
-
-			/** Adds a band which calculates the band value as a pesudo constant q band derived from an FFT power spectrum.
-			 *
-			 * @param InCenterFrequency - Frequency of interest in hz.
-			 * @param InQFactor - QFactor used in band calculation. QFactor = CenterFreq / BandWidth. A small QFactor results in a wide band.
-			 * @param InMetric - Metric used to calculate return value.
-			 * @param InDecibelNoiseFloor - If the metric is Decibel, this is the minimum decibel value allowed.
-			 * @param bInDoNormalize - If true, all values are scaled and clamped between 0.0 and 1.f. In the 
-			 *                         case of Decibels, 0.0 corresponds to the decibel noise floor and 1.f to 0dB.
-			 */
-			virtual void AddConstantQBand(float InCenterFrequency, float InQFactor, EMetric InMetric = EMetric::Decibel, float InDecibelNoiseFloor=-40.f, bool bInDoNormalize=true) = 0;
+			/** Adds a band to extract based on the given settings. */
+			virtual void AddBand(const FBandSettings& InSettings) = 0;
 
 			/** Extract the bands from a complex frequency buffer.
 			 *
 			 * @param InComplexBuffer - Buffer of complex frequency data from a FFT.
+			 * @param InTimestamp - A timestamp associated with the input complex buffer.
 			 * @param OutValues - Array to store output bands.
 			 */
-			virtual void ExtractBands(const AlignedFloatBuffer& InComplexBuffer, TArray<float>& OutValues) = 0;
+			virtual void ExtractBands(const AlignedFloatBuffer& InComplexBuffer, double InTimestamp, TArray<float>& OutValues) = 0;
 
 			/** Creates a ISpectrumBandExtractor. */
 			static TUniquePtr<ISpectrumBandExtractor> CreateSpectrumBandExtractor(const FSpectrumBandExtractorSettings& InSettings);
@@ -210,14 +217,21 @@ namespace Audio
 
 		// Input. Used on analysis thread to lock a buffer to write to.
 		AlignedFloatBuffer& StartWorkOnBuffer();
-		void StopWorkOnBuffer();
+
+		// When calling stop work on buffer, also set timestmap associated with buffer.
+		void StopWorkOnBuffer(double InTimestamp);
 		
 		// Output. Used to lock the most recent buffer we analyzed.
 		const AlignedFloatBuffer& LockMostRecentBuffer() const;
+
+		// Output. Used to lock the most recent buffer we analyzed.
+		// OutTimestamp is populated with the timestamp associated with the buffer wehn StopWorkOnBuffer is called.
+		const AlignedFloatBuffer& LockMostRecentBuffer(double& OutTimestamp) const;
 		void UnlockBuffer();
 
 	private:
 		TArray<AlignedFloatBuffer> ComplexBuffers;
+		TArray<double> Timestamps;
 
 		// Private functions. Either increments or decrements the respective counter,
 		// based on which index is currently in use. Mutually locked.
@@ -265,6 +279,15 @@ namespace Audio
 	class SIGNALPROCESSING_API FSpectrumAnalyzer
 	{
 	public:
+		// Peak interpolation method. If the EFFTSize is small but will be densely sampled,
+		// it's worth using a linear or quadratic interpolation method.
+		enum class EPeakInterpolationMethod : uint8
+		{
+			NearestNeighbor,
+			Linear,
+			Quadratic
+		};
+
 		// If an instance is created using the default constructor, Init() must be called before it is used.
 		FSpectrumAnalyzer();
 
@@ -285,10 +308,10 @@ namespace Audio
 		void GetSettings(FSpectrumAnalyzerSettings& OutSettings);
 
 		// Samples magnitude (linearly) for a given frequency, in Hz.
-		float GetMagnitudeForFrequency(float InFrequency);
+		float GetMagnitudeForFrequency(float InFrequency, EPeakInterpolationMethod InMethod = EPeakInterpolationMethod::Linear);
 
 		// Samples phase for a given frequency, in Hz.
-		float GetPhaseForFrequency(float InFrequency);
+		float GetPhaseForFrequency(float InFrequency, EPeakInterpolationMethod InMethod = EPeakInterpolationMethod::Linear);
 
 		// Return array of bands using spectrum band extractor.
 		void GetBands(ISpectrumBandExtractor& InExtractor, TArray<float>& OutValues);
@@ -320,7 +343,7 @@ namespace Audio
 		void ResetSettings();
 
 		// Called in GetMagnitudeForFrequency and GetPhaseForFrequency.
-		void PerformInterpolation(const AlignedFloatBuffer& InComplexBuffer, FSpectrumAnalyzerSettings::EPeakInterpolationMethod InMethod, const float InFreq, float& OutReal, float& OutImag);
+		void PerformInterpolation(const AlignedFloatBuffer& InComplexBuffer, EPeakInterpolationMethod InMethod, const float InFreq, float& OutReal, float& OutImag);
 
 		// Cached current settings. Only actually used in ResetSettings().
 		FSpectrumAnalyzerSettings CurrentSettings;
@@ -337,10 +360,12 @@ namespace Audio
 		EFFTScaling FFTScaling;
 
 		AlignedFloatBuffer AnalysisTimeDomainBuffer;
+		FThreadSafeCounter SampleCounter; 
 		TCircularAudioBuffer<float> InputQueue;
 		FSpectrumAnalyzerBuffer FrequencyBuffer;
 
 		// if non-null, owns pointer to locked frequency vector we're using.
+		double LockedBufferTimestamp;
 		const AlignedFloatBuffer* LockedFrequencyVector;
 
 		// This is used if PerformAnalysisIfPossible is called
