Adding first classes on pitch detector. Basic FFT and ACF pitch detector.

diff --git a/Engine/Plugins/Runtime/AudioSynesthesia/Source/AudioSynesthesiaCore/Private/AudioSynesthesiaCoreLog.h b/Engine/Plugins/Runtime/AudioSynesthesia/Source/AudioSynesthesiaCore/Private/AudioSynesthesiaCoreLog.h
new file mode 100644
--- /dev/null
+++ b/Engine/Plugins/Runtime/AudioSynesthesia/Source/AudioSynesthesiaCore/Private/AudioSynesthesiaCoreLog.h
@@ -0,0 +1,7 @@
+// Copyright Epic Games, Inc. All Rights Reserved.
+
+#pragma once
+
+#include "CoreMinimal.h"
+
+DECLARE_LOG_CATEGORY_EXTERN(LogAudioSynesthesiaCore, Log, All);
diff --git a/Engine/Plugins/Runtime/AudioSynesthesia/Source/AudioSynesthesiaCore/Private/AudioSynesthesiaCoreModule.cpp b/Engine/Plugins/Runtime/AudioSynesthesia/Source/AudioSynesthesiaCore/Private/AudioSynesthesiaCoreModule.cpp
--- a/Engine/Plugins/Runtime/AudioSynesthesia/Source/AudioSynesthesiaCore/Private/AudioSynesthesiaCoreModule.cpp
+++ b/Engine/Plugins/Runtime/AudioSynesthesia/Source/AudioSynesthesiaCore/Private/AudioSynesthesiaCoreModule.cpp
@@ -1,6 +1,7 @@
 // Copyright Epic Games, Inc. All Rights Reserved.
 
 #include "AudioSynesthesiaCoreModule.h"
+#include "AudioSynesthesiaCoreLog.h"
 #include "CoreMinimal.h"
 #include "Modules/ModuleInterface.h"
 #include "Modules/ModuleManager.h"
@@ -9,6 +10,7 @@
 #include "ConstantQNRTFactory.h"
 #include "OnsetNRTFactory.h"
 
+DEFINE_LOG_CATEGORY(LogAudioSynesthesiaCore);
 
 namespace Audio
 {
diff --git a/Engine/Plugins/Runtime/AudioSynesthesia/Source/AudioSynesthesiaCore/Private/AutoCorrelationPitchDetector.cpp b/Engine/Plugins/Runtime/AudioSynesthesia/Source/AudioSynesthesiaCore/Private/AutoCorrelationPitchDetector.cpp
new file mode 100644
--- /dev/null
+++ b/Engine/Plugins/Runtime/AudioSynesthesia/Source/AudioSynesthesiaCore/Private/AutoCorrelationPitchDetector.cpp
@@ -0,0 +1,170 @@
+// Copyright Epic Games, Inc. All Rights Reserved.
+
+#include "AutoCorrelationPitchDetector.h"
+
+#include "AudioSynesthesiaCoreLog.h"
+#include "CoreMinimal.h"
+#include "PeakPicker.h"
+#include "DSP/BlockCorrelator.h"
+#include "DSP/SlidingWindow.h"
+#include "DSP/FloatArrayMath.h"
+#include "SignalProcessingModule.h"
+
+namespace Audio
+{
+	namespace AutoCorrelationPitchDetectorIntrinsics
+	{
+		constexpr int32 MinimumWindowSize = 512;
+		constexpr int32 MaximumWindowSize = 8192;
+	}
+
+	FAutoCorrelationPitchDetector::FAutoCorrelationPitchDetector(const FAutoCorrelationPitchDetectorSettings& InSettings, float InSampleRate)
+	:	Settings(InSettings)
+	,	SampleRate(InSampleRate)
+	,	MinAutoCorrBin(0)
+	,	MaxAutoCorrBin(0)
+	,	WindowCounter(0)
+	,	SlidingBuffer(1, 1)
+	{
+		using namespace AutoCorrelationPitchDetectorIntrinsics;
+
+		if (!ensure(SampleRate > 0.f))
+		{
+			SampleRate = 1.f;
+		}
+
+		Settings.MinimumFrequency = FMath::Max(Settings.MinimumFrequency, 1.f);
+		Settings.MaximumFrequency = FMath::Max(Settings.MaximumFrequency, 1.f);
+		
+		// setup sliding window
+		int32 MinFreqNumSamples = FMath::RoundToInt(3.f * SampleRate / Settings.MinimumFrequency);
+		MinFreqNumSamples = FMath::Max(MinFreqNumSamples, MinimumWindowSize);
+
+		if (MinFreqNumSamples > MaximumWindowSize)
+		{
+			UE_LOG(LogAudioSynesthesiaCore, Warning, TEXT("Pitches at minimum frequency (%f hz) may not be detected due to limit of internal window size"), Settings.MinimumFrequency);
+			MinFreqNumSamples = MaximumWindowSize;
+		}
+
+		int32 Log2WindowNumSamples = 1;
+		while ((1 << Log2WindowNumSamples) < MinFreqNumSamples)
+		{
+			Log2WindowNumSamples++;
+		}
+
+		int32 HopNumSamples = FMath::Max(1, FMath::RoundToInt(SampleRate * Settings.AnalysisHopSeconds));
+
+		int32 WindowNumSamples = 1 << Log2WindowNumSamples;
+
+		SlidingBuffer = TSlidingBuffer<float>(WindowNumSamples, HopNumSamples);
+		WindowBuffer.AddUninitialized(WindowNumSamples);
+
+		// Setup range of result search
+		MinAutoCorrBin = FMath::Clamp(FMath::FloorToInt(SampleRate / Settings.MaximumFrequency), 1, WindowNumSamples);
+		MaxAutoCorrBin = FMath::Clamp(FMath::CeilToInt(SampleRate / Settings.MinimumFrequency), 1, WindowNumSamples);
+
+
+		// setup peak picker
+		float InverseSensitivity = 1.f - Settings.Sensitivity;
+		int32 AutoCorrRange = FMath::Max(1, MaxAutoCorrBin - MinAutoCorrBin);
+		
+		FPeakPickerSettings PeakPickerSettings;
+		PeakPickerSettings.NumPreMax = FMath::Max(1, FMath::RoundToInt(InverseSensitivity * AutoCorrRange / 2.f));
+		PeakPickerSettings.NumPostMax = PeakPickerSettings.NumPreMax;
+		PeakPickerSettings.NumPreMean = FMath::Max(1, FMath::RoundToInt(InverseSensitivity * AutoCorrRange));
+		PeakPickerSettings.NumPostMean = PeakPickerSettings.NumPreMean;
+		PeakPickerSettings.NumWait = 1;
+		PeakPickerSettings.MeanDelta = 0.01f + InverseSensitivity;
+
+		PeakPicker = MakeUnique<FPeakPicker>(PeakPickerSettings);
+
+		// setup block correlator
+		FBlockCorrelatorSettings BlockCorrelatorSettings;
+
+		BlockCorrelatorSettings.Log2NumValuesInBlock = Log2WindowNumSamples;
+		BlockCorrelatorSettings.WindowType = EWindowType::Blackman;
+		BlockCorrelatorSettings.bDoNormalize = true;
+
+		Correlator = MakeUnique<FBlockCorrelator>(BlockCorrelatorSettings);
+
+		if (Correlator.IsValid())
+		{
+			int32 AutoCorrNumSamples = Correlator->GetNumOutputValues();
+			if (AutoCorrNumSamples > 0)
+			{
+				AutoCorrBuffer.AddUninitialized(AutoCorrNumSamples);
+			}
+		}
+	}
+
+	FAutoCorrelationPitchDetector::~FAutoCorrelationPitchDetector()
+	{
+	}
+
+	void FAutoCorrelationPitchDetector::DetectPitches(const AlignedFloatBuffer& InMonoAudio, TArray<FPitchInfo>& OutPitches)
+	{
+		TAutoSlidingWindow<float, FAudioBufferAlignedAllocator> SlidingWindow(SlidingBuffer, InMonoAudio, WindowBuffer);
+
+		const float InverseSampleRate = 1.f / SampleRate;
+
+		for (AlignedFloatBuffer& AnalysisBuffer : SlidingWindow)
+		{
+			int32 SampleCounter = (WindowCounter * SlidingBuffer.GetNumHopSamples()) + SlidingBuffer.GetNumWindowSamples() / 2;
+			const float Timestamp = SampleCounter * InverseSampleRate;
+
+			WindowCounter += 1;
+
+			if (!Correlator.IsValid())
+			{
+				continue;
+			}
+
+			Correlator->AutoCorrelate(AnalysisBuffer, AutoCorrBuffer);
+
+			// Pitch peaks
+			PeakIndices.Reset();
+			PeakPicker->PickPeaks(AutoCorrBuffer, PeakIndices);
+
+			for (int32 Index : PeakIndices)
+			{
+				if ((Index < MinAutoCorrBin) || (Index > MaxAutoCorrBin))
+				{
+					continue;
+				}
+
+				float PeakLocOffset = 0.f;
+				float PeakLocMaximum = 0.f;
+
+				const float* AutoCorrData = AutoCorrBuffer.GetData();
+				float Energy = FMath::Max(SMALL_NUMBER, AutoCorrData[0]);
+
+				// Perform quadratic interpolation on peak.
+				// If this returns false, then location is not a true peak.
+				if (QuadraticPeakInterpolation(&AutoCorrData[Index - 1], PeakLocOffset, PeakLocMaximum))
+				{
+					float PeakLoc = static_cast<float>(Index) + PeakLocOffset;
+
+					if (PeakLoc < 2.f)
+					{
+						// Invalid peak location.
+						continue;
+					}
+
+					FPitchInfo Info;
+
+					Info.Frequency = SampleRate / PeakLoc;
+					Info.Strength = FMath::Clamp(PeakLocMaximum / Energy, 0.0f, 1.f);
+					Info.Timestamp = Timestamp;
+
+					OutPitches.Add(MoveTemp(Info));
+				}
+			}
+		}
+	}
+
+	void FAutoCorrelationPitchDetector::Finalize(TArray<FPitchInfo>& OutPitches)
+	{
+		SlidingBuffer.Reset();
+		WindowCounter = 0;
+	}
+}
diff --git a/Engine/Plugins/Runtime/AudioSynesthesia/Source/AudioSynesthesiaCore/Private/FFTPeakPitchDetector.cpp b/Engine/Plugins/Runtime/AudioSynesthesia/Source/AudioSynesthesiaCore/Private/FFTPeakPitchDetector.cpp
new file mode 100644
--- /dev/null
+++ b/Engine/Plugins/Runtime/AudioSynesthesia/Source/AudioSynesthesiaCore/Private/FFTPeakPitchDetector.cpp
@@ -0,0 +1,189 @@
+// Copyright Epic Games, Inc. All Rights Reserved.
+
+#include "FFTPeakPitchDetector.h"
+
+#include "AudioSynesthesiaCoreLog.h"
+#include "CoreMinimal.h"
+#include "PeakPicker.h"
+#include "DSP/Dsp.h"
+#include "DSP/FFTAlgorithm.h"
+#include "DSP/SlidingWindow.h"
+#include "DSP/FloatArrayMath.h"
+
+namespace Audio
+{
+	FFFTPeakPitchDetector::FFFTPeakPitchDetector(const FFFTPeakPitchDetectorSettings& InSettings, float InSampleRate)
+	:	Settings(InSettings)
+	,	SampleRate(InSampleRate)
+	,	FFTSize(1 << InSettings.Log2FFTSize)
+	,	MinFFTBin(1)
+	,	MaxFFTBin(1)
+	,	FFTScaling(1.f)
+	,	WindowCounter(0)
+	,	SlidingBuffer(1, 1)
+	{
+		if (!ensure(SampleRate > 0.f))
+		{
+			SampleRate = 1.f;
+		}
+
+		int32 AnalysisWindowNumSamples = FMath::RoundToInt(Settings.AnalysisWindowSeconds * SampleRate);
+		int32 AnalysisHopNumSamples = FMath::RoundToInt(Settings.AnalysisHopSeconds * SampleRate);
+
+		if (AnalysisWindowNumSamples > FFTSize)
+		{
+			UE_LOG(LogAudioSynesthesiaCore, Log, TEXT("The analysis window size [%d] is being reduced to the FFT size [%d]"), AnalysisWindowNumSamples, FFTSize);
+			AnalysisWindowNumSamples = FFTSize;
+		}
+
+		UE_CLOG(AnalysisHopNumSamples > AnalysisWindowNumSamples, LogAudioSynesthesiaCore, Log, TEXT("Some samples will not be anlayzed because the hop size [%d] is larger than the analysis window [%d]"), AnalysisHopNumSamples, AnalysisWindowNumSamples);
+
+		SlidingBuffer = TSlidingBuffer<float>(AnalysisWindowNumSamples, AnalysisHopNumSamples);
+
+		// setup fft
+		FFFTSettings FFTSettings;
+		FFTSettings.Log2Size = Settings.Log2FFTSize;
+		FFTSettings.bArrays128BitAligned = true;
+		FFTSettings.bEnableHardwareAcceleration = true;
+
+		FFTAlgorithm = FFFTFactory::NewFFTAlgorithm(FFTSettings);
+		check(FFTAlgorithm.IsValid());
+
+		switch (FFTAlgorithm->ForwardScaling())
+		{
+
+			case EFFTScaling::MultipliedByFFTSize:
+				FFTScaling = 1.f / FFTAlgorithm->Size();
+				break;
+
+			case EFFTScaling::MultipliedBySqrtFFTSize:
+				FFTScaling = 1.f / FMath::Sqrt(static_cast<float>(FFTAlgorithm->Size()));
+				break;
+
+			case EFFTScaling::DividedByFFTSize:
+				FFTScaling = FFTAlgorithm->Size();
+				break;
+
+			case EFFTScaling::DividedBySqrtFFTSize:
+				FFTScaling = FMath::Sqrt(static_cast<float>(FFTAlgorithm->Size()));
+				break;
+
+			case EFFTScaling::None:
+			default:
+				FFTScaling = 1.f;
+				break;
+		}
+
+		float HzPerBin = SampleRate / FFTSize;
+		MinFFTBin = FMath::Max(1, FMath::FloorToInt(Settings.MinimumFrequency / HzPerBin));
+		MaxFFTBin = FMath::Max(1, FMath::CeilToInt(Settings.MaximumFrequency / HzPerBin));
+		int32 FFTRange = MaxFFTBin - MinFFTBin + 1;
+
+		// setup peak picker
+		float InverseSensitivity = 1.f - Settings.Sensitivity;
+		
+		FPeakPickerSettings PeakPickerSettings;
+		PeakPickerSettings.NumPreMax = FMath::Max(1, FMath::RoundToInt(InverseSensitivity * FFTRange / 2.f));
+		PeakPickerSettings.NumPostMax = PeakPickerSettings.NumPreMax;
+		PeakPickerSettings.NumPreMean = FMath::Max(1, FMath::RoundToInt(InverseSensitivity * FFTRange));
+		PeakPickerSettings.NumPostMean = PeakPickerSettings.NumPreMean;
+		PeakPickerSettings.NumWait = 1;
+		PeakPickerSettings.MeanDelta = 0.01f + InverseSensitivity;
+
+		PeakPicker = MakeUnique<FPeakPicker>(PeakPickerSettings);
+
+		// setup some extra buffers
+		if (FFTAlgorithm->NumInputFloats() > 0)
+		{
+			ZeroPaddedAnalysisBuffer.AddZeroed(FFTAlgorithm->NumInputFloats());
+		}
+
+		if (FFTAlgorithm->NumOutputFloats() > 0)
+		{
+			ComplexSpectrumBuffer.AddZeroed(FFTAlgorithm->NumOutputFloats());
+
+			PowerSpectrumBuffer.AddZeroed(FMath::Max(1, FFTAlgorithm->NumOutputFloats() / 2));
+		}
+	}
+
+	FFFTPeakPitchDetector::~FFFTPeakPitchDetector()
+	{
+	}
+
+	void FFFTPeakPitchDetector::DetectPitches(const AlignedFloatBuffer& InMonoAudio, TArray<FPitchInfo>& OutPitches)
+	{
+		TAutoSlidingWindow<float, FAudioBufferAlignedAllocator> SlidingWindow(SlidingBuffer, InMonoAudio, WindowBuffer);
+
+		const float InverseMultiply = 1.f / SampleRate;
+
+		for (AlignedFloatBuffer& AnalysisBuffer : SlidingWindow)
+		{
+			int32 SampleCounter = (WindowCounter * SlidingBuffer.GetNumHopSamples()) + SlidingBuffer.GetNumWindowSamples() / 2;
+			float Timestamp = SampleCounter * InverseMultiply;
+
+			WindowCounter += 1;
+
+			// FFT
+			FMemory::Memcpy(ZeroPaddedAnalysisBuffer.GetData(), AnalysisBuffer.GetData(), AnalysisBuffer.Num() * sizeof(float));
+			FFTAlgorithm->ForwardRealToComplex(ZeroPaddedAnalysisBuffer.GetData(), ComplexSpectrumBuffer.GetData());
+			ArrayMultiplyByConstantInPlace(ComplexSpectrumBuffer, FFTScaling);
+
+			ArrayComplexToPower(ComplexSpectrumBuffer, PowerSpectrumBuffer);
+
+			// Pitch peaks
+			TArray<int32> PeakIndices;
+			PeakPicker->PickPeaks(PowerSpectrumBuffer, PeakIndices);
+
+			PeakIndices = PeakIndices.FilterByPredicate([&](int32 Index){
+				return (Index >= MinFFTBin) && (Index <= MaxFFTBin);
+			});
+
+
+			for (int32 Index : PeakIndices)
+			{
+				float PeakLocOffset = 0.f;
+				float PeakLocMaximum = 0.f;
+
+				const float* PowerSpectrumData = PowerSpectrumBuffer.GetData();
+
+				// Perform quadratic interpolation on peak.
+				// If this returns false, then location is not a true peak.
+				if (QuadraticPeakInterpolation(&PowerSpectrumData[Index - 1], PeakLocOffset, PeakLocMaximum))
+				{
+					FPitchInfo Info;
+
+					Info.Frequency = (static_cast<float>(Index) + PeakLocOffset) * SampleRate / static_cast<float>(FFTAlgorithm->Size());
+					Info.Strength = PeakLocMaximum;
+					Info.Timestamp = Timestamp;
+
+					OutPitches.Add(MoveTemp(Info));
+				}
+			}
+
+			// Get median of work buffer
+			if (PowerSpectrumBuffer.Num() > 0)
+			{
+				// The strength adjustment is applied after the locations are determined because it requires us
+				// to sort the 
+				PowerSpectrumBuffer.Sort();
+
+				float MedianValue = FMath::Max(SMALL_NUMBER, PowerSpectrumBuffer[PowerSpectrumBuffer.Num() / 2]);
+
+				for (FPitchInfo& Info : OutPitches)
+				{
+					float DecibelAboveMedian = 10.f * FMath::LogX(10.f, FMath::Max(Info.Strength / MedianValue, SMALL_NUMBER));
+					
+					Info.Strength = FMath::Clamp(DecibelAboveMedian / Settings.MaxStrengthSNRDecibels, 0.0f, 1.f);
+				}
+
+			}
+
+		}
+	}
+
+	void FFFTPeakPitchDetector::Finalize(TArray<FPitchInfo>& OutPitches)
+	{
+		SlidingBuffer.Reset();
+		WindowCounter = 0;
+	}
+}
diff --git a/Engine/Plugins/Runtime/AudioSynesthesia/Source/AudioSynesthesiaCore/Private/MaxStrengthPitchTracker.cpp b/Engine/Plugins/Runtime/AudioSynesthesia/Source/AudioSynesthesiaCore/Private/MaxStrengthPitchTracker.cpp
new file mode 100644
--- /dev/null
+++ b/Engine/Plugins/Runtime/AudioSynesthesia/Source/AudioSynesthesiaCore/Private/MaxStrengthPitchTracker.cpp
@@ -0,0 +1,150 @@
+// Copyright Epic Games, Inc. All Rights Reserved.
+
+#include "MaxStrengthPitchTracker.h"
+
+#include "Algo/MaxElement.h"
+#include "AudioSynesthesiaCoreLog.h"
+#include "CoreMinimal.h"
+
+namespace Audio
+{
+	namespace MaxStrengthPitchTrackIntrinsics
+	{
+
+		// Likely this method will be useful elsewhere
+		template<typename T, class PROJECTION_CLASS, class ON_GROUP_CLASS>
+		void GroupBy(TArrayView<const T> InData, const PROJECTION_CLASS& Proj, const ON_GROUP_CLASS& OnGroup)
+		{
+			const int32 Num = InData.Num();
+
+			if (Num > 0)
+			{
+				const T* Data = InData.GetData();
+
+				int32 GroupStartIndex = 0;
+				auto Key = Proj(Data[0]);
+
+				for (int32 i = 1; i < Num; i++)
+				{
+					auto NewKey = Proj(Data[i]);
+					if (Key != NewKey)
+					{
+						OnGroup(Key, InData.Slice(GroupStartIndex, i - GroupStartIndex));
+
+						Key = NewKey;
+						GroupStartIndex = i;
+					}
+				}
+
+				OnGroup(Key, InData.Slice(GroupStartIndex, Num - GroupStartIndex));
+			}
+		}
+	}
+
+	FMaxStrengthPitchTracker::FMaxStrengthPitchTracker(const FMaxStrengthPitchTrackerSettings& InSettings, TUniquePtr<IPitchDetector> InDetector)
+	:	Settings(InSettings)
+	,	PitchDetector(MoveTemp(InDetector))
+	{
+		UE_CLOG(!PitchDetector.IsValid(), LogAudioSynesthesiaCore, Warning, TEXT("Invalid pointer to pitch detector. MaxStrengthPitchDetector will not return valid results."));
+	}
+
+	FMaxStrengthPitchTracker::~FMaxStrengthPitchTracker()
+	{
+	}
+
+	void FMaxStrengthPitchTracker::TrackPitches(const AlignedFloatBuffer& InMonoAudio, TArray<FPitchTrackInfo>& OutPitchTracks)
+	{
+		if (PitchDetector.IsValid())
+		{
+			TArray<FPitchInfo> NewObservations;
+
+			PitchDetector->DetectPitches(InMonoAudio, NewObservations);
+
+			Observations.Append(MoveTemp(NewObservations));
+		}
+	}
+
+	void FMaxStrengthPitchTracker::Finalize(TArray<FPitchTrackInfo>& OutPitchTracks)
+	{
+		using namespace MaxStrengthPitchTrackIntrinsics;
+
+		if (PitchDetector.IsValid())
+		{
+			PitchDetector->Finalize(Observations);
+		}
+
+		// Sort obersvations temporally
+		Observations.Sort([](const FPitchInfo& A, const FPitchInfo& B) { return A.Timestamp < B.Timestamp; });
+
+		TArray<FPitchInfo> MaxObservations;
+
+		TArrayView<const FPitchInfo> ObservationsView(Observations);
+
+		// Group by pitch infos with equal timestamps
+		GroupBy(ObservationsView, [](const FPitchInfo& Info) { return Info.Timestamp; }, [&](float Timestamp, TArrayView<const FPitchInfo> PitchInfos)
+			{
+				// Get maximum strength element within all elements with equal timestamps
+				const FPitchInfo* MaxInfo = Algo::MaxElementBy(PitchInfos, [](const FPitchInfo& Info) { return Info.Strength; });
+
+				if (nullptr != MaxInfo)
+				{
+					MaxObservations.Add(*MaxInfo);
+				}
+			}
+		);
+
+		// Build tracks if requirements are met.
+		FPitchTrackInfo TrackInfo;
+
+		for (const FPitchInfo& PitchInfo : MaxObservations)
+		{
+			bool bEndPreviousTrack = false;
+
+			bool bAddToTrack = PitchInfo.Strength >= Settings.MinimumStrength;
+
+			if (TrackInfo.Observations.Num() > 0)
+			{
+				if (!bAddToTrack)
+				{
+					// End previous track if new pitch observations is too weak.
+					bEndPreviousTrack = true;
+				}
+				else
+				{
+					const FPitchInfo& Last = TrackInfo.Observations.Last(0);
+					if (Last.Frequency == 0.f)
+					{
+						// A frequency of zero screws up the frequency deviation calculation. Also, DC is not a pitch.
+						bEndPreviousTrack = true;
+					}
+					else
+					{
+						// Check if frequency ratio deviation below threshold
+						float FrequencyRatioDeviation = FMath::Abs(1.f - PitchInfo.Frequency / Last.Frequency);
+						bEndPreviousTrack |= (FrequencyRatioDeviation > Settings.MaximumFrequencyRatioDeviation);
+					}
+				}
+			}
+
+			if (bEndPreviousTrack)
+			{
+				OutPitchTracks.Add(MoveTemp(TrackInfo));
+
+				TrackInfo.Observations.Reset();
+			}
+
+			if (bAddToTrack)
+			{
+				TrackInfo.Observations.Add(PitchInfo);
+			}
+		}
+
+		if (TrackInfo.Observations.Num() > 0)
+		{
+			OutPitchTracks.Add(MoveTemp(TrackInfo));
+		}
+
+		// Clear out observations for next audio stream.
+		Observations.Reset();
+	}
+}
diff --git a/Engine/Plugins/Runtime/AudioSynesthesia/Source/AudioSynesthesiaCore/Public/AutoCorrelationPitchDetector.h b/Engine/Plugins/Runtime/AudioSynesthesia/Source/AudioSynesthesiaCore/Public/AutoCorrelationPitchDetector.h
new file mode 100644
--- /dev/null
+++ b/Engine/Plugins/Runtime/AudioSynesthesia/Source/AudioSynesthesiaCore/Public/AutoCorrelationPitchDetector.h
@@ -0,0 +1,72 @@
+// Copyright Epic Games, Inc. All Rights Reserved.
+
+#pragma once
+
+#include "CoreMinimal.h"
+#include "PitchTracker.h"
+#include "PeakPicker.h"
+#include "DSP/BlockCorrelator.h"
+#include "DSP/SlidingWindow.h"
+
+namespace Audio
+{
+	/** Settings for creating an autocorrelation pitch detector. */
+	struct FAutoCorrelationPitchDetectorSettings
+	{
+		/** Time (in seconds) between analysis windows. */
+		float AnalysisHopSeconds = 0.01f;
+
+		/** Minimum frequency of pitch result. */
+		float MinimumFrequency = 20.f;
+
+		/** Maximum frequency of pitch result. */
+		float MaximumFrequency = 10000.f;
+
+		/** Sensitivity of peak detection */
+		float Sensitivity = 0.9f;
+	};
+
+	/** Pitch detector based on autocorrelation. Note that autocorrelation pitch detectors 
+	 * give more accurate frequency results for low frequencies, but have issues with octave errors.
+	 * Generally, it will produce erroneous frequency observations in octaves _below_ the true
+	 * pitch frequency.
+	 */
+	class AUDIOSYNESTHESIACORE_API FAutoCorrelationPitchDetector : public IPitchDetector
+	{
+		public:
+			/** Create an auto correlation pitch detector with settings and a sample rate. */
+			FAutoCorrelationPitchDetector(const FAutoCorrelationPitchDetectorSettings& InSettings, float InSampleRate);
+
+			virtual ~FAutoCorrelationPitchDetector();
+
+			/** Detect pitches in the audio. This can be called repeatedly with new audio.
+			 *
+			 * @param InMonoAudio - A mono audio buffer of any length.
+			 * @param OutPitches - This array is filled with pitch observations.
+			 */
+			virtual void DetectPitches(const AlignedFloatBuffer& InMonoAudio, TArray<FPitchInfo>& OutPitches) override;
+
+			/** Resets internal audio buffers. This pitch detector does not produce any more pitches on Finalize. */
+			virtual void Finalize(TArray<FPitchInfo>& OutPitches) override;
+
+		private:
+
+			FAutoCorrelationPitchDetectorSettings Settings;
+
+			float SampleRate;
+			int32 MinAutoCorrBin;
+			int32 MaxAutoCorrBin;
+			int32 WindowCounter;
+
+			TSlidingBuffer<float> SlidingBuffer;
+			AlignedFloatBuffer WindowBuffer;
+			AlignedFloatBuffer AutoCorrBuffer;
+			TUniquePtr<FBlockCorrelator> Correlator;
+
+			TArray<int32> PeakIndices;
+			TUniquePtr<FPeakPicker> PeakPicker;
+
+	};
+}
+
+
diff --git a/Engine/Plugins/Runtime/AudioSynesthesia/Source/AudioSynesthesiaCore/Public/FFTPeakPitchDetector.h b/Engine/Plugins/Runtime/AudioSynesthesia/Source/AudioSynesthesiaCore/Public/FFTPeakPitchDetector.h
new file mode 100644
--- /dev/null
+++ b/Engine/Plugins/Runtime/AudioSynesthesia/Source/AudioSynesthesiaCore/Public/FFTPeakPitchDetector.h
@@ -0,0 +1,77 @@
+// Copyright Epic Games, Inc. All Rights Reserved.
+
+#pragma once
+
+#include "CoreMinimal.h"
+#include "PitchTracker.h"
+#include "PeakPicker.h"
+#include "DSP/BufferVectorOperations.h"
+#include "DSP/FFTAlgorithm.h"
+#include "DSP/SlidingWindow.h"
+
+namespace Audio
+{
+	/** Settings for FFT base pitch detection */
+	struct FFFTPeakPitchDetectorSettings
+	{
+		/** Log 2 size of the FFT */
+		int32 Log2FFTSize = 13; // 2^13 = 8192
+
+		/** Number of seconds represented in each analysis window. */
+		float AnalysisWindowSeconds = 0.1f;
+
+		/** Number of seconds between analysis windows. */
+		float AnalysisHopSeconds = 0.01f;
+
+		/** Minimum frequency of detected pitches. */
+		float MinimumFrequency = 20.f;
+
+		/** Maximum frequency of detected pitches. */
+		float MaximumFrequency = 10000.f;
+
+		/** Sensitivity used for picking peaks in the power spectrum. */
+		float Sensitivity = 0.9f;
+
+		/** Reference SNR for calculating strength of signal. */
+		float MaxStrengthSNRDecibels = 20.f;
+	};
+
+	/** FFT based pitch detector accel at detecting pitches at higher frequencies. */
+	class AUDIOSYNESTHESIACORE_API FFFTPeakPitchDetector : public IPitchDetector
+	{
+		public:
+			/** Create an fft pitch detector with settings and a sample rate. */
+			FFFTPeakPitchDetector(const FFFTPeakPitchDetectorSettings& InSettings, float InSampleRate);
+
+			virtual ~FFFTPeakPitchDetector();
+
+			/** Detect pitches in the audio. This can be called repeatedly with new audio.
+			 *
+			 * @param InMonoAudio - A mono audio buffer of any length.
+			 * @param OutPitches - This array is filled with pitch observations.
+			 */
+			virtual void DetectPitches(const AlignedFloatBuffer& InMonoAudio, TArray<FPitchInfo>& OutPitches) override;
+
+			/** Resets internal audio buffers. This pitch detector does not produce any more pitches on Finalize. */
+			virtual void Finalize(TArray<FPitchInfo>& OutPitches) override;
+
+		private:
+			FFFTPeakPitchDetectorSettings Settings;
+
+			float SampleRate;
+			int32 FFTSize;
+			int32 MinFFTBin;
+			int32 MaxFFTBin;
+			float FFTScaling;
+			int32 WindowCounter;
+
+			TUniquePtr<IFFTAlgorithm> FFTAlgorithm;
+			TUniquePtr<FPeakPicker> PeakPicker;
+
+			TSlidingBuffer<float> SlidingBuffer;
+			AlignedFloatBuffer WindowBuffer;
+			AlignedFloatBuffer ZeroPaddedAnalysisBuffer;
+			AlignedFloatBuffer ComplexSpectrumBuffer;
+			AlignedFloatBuffer PowerSpectrumBuffer;
+	};
+}
diff --git a/Engine/Plugins/Runtime/AudioSynesthesia/Source/AudioSynesthesiaCore/Public/MaxStrengthPitchTracker.h b/Engine/Plugins/Runtime/AudioSynesthesia/Source/AudioSynesthesiaCore/Public/MaxStrengthPitchTracker.h
new file mode 100644
--- /dev/null
+++ b/Engine/Plugins/Runtime/AudioSynesthesia/Source/AudioSynesthesiaCore/Public/MaxStrengthPitchTracker.h
@@ -0,0 +1,42 @@
+// Copyright Epic Games, Inc. All Rights Reserved.
+
+#pragma once
+
+#include "CoreMinimal.h"
+#include "PitchTracker.h"
+#include "DSP/BufferVectorOperations.h"
+
+namespace Audio
+{
+	/** Settings for FMaxStrengthPitchTracker */
+	struct FMaxStrengthPitchTrackerSettings
+	{
+		/** Minimum strength of pitch to be included in pitch track. */
+		float MinimumStrength = 0.1f;
+
+		/** Maximum deviation of ratio between two freuqencies to consider them part of the same pitch track. */
+		float MaximumFrequencyRatioDeviation = 0.05f;
+	};
+
+	/** Track the maximum strength pitch. */
+	class AUDIOSYNESTHESIACORE_API FMaxStrengthPitchTracker : public IPitchTracker
+	{
+		public:
+			/** Constructor accepts settings and a unique pointer to the pitch detector. */
+			FMaxStrengthPitchTracker(const FMaxStrengthPitchTrackerSettings& InSettings, TUniquePtr<IPitchDetector> InDetector);
+
+			virtual ~FMaxStrengthPitchTracker();
+
+			/** Append pitch tracks to processed audio. */
+			virtual void TrackPitches(const AlignedFloatBuffer& InMonoAudio, TArray<FPitchTrackInfo>& OutPitchTracks) override;
+
+			/** Call when all audio has been processed. Appends pitch tracks to OutPitchTracks. */
+			virtual void Finalize(TArray<FPitchTrackInfo>& OutPitchTracks) override;
+
+		private:
+			FMaxStrengthPitchTrackerSettings Settings;
+			TUniquePtr<IPitchDetector> PitchDetector;
+
+			TArray<FPitchInfo> Observations;
+	};
+}
diff --git a/Engine/Plugins/Runtime/AudioSynesthesia/Source/AudioSynesthesiaCore/Public/PitchTracker.h b/Engine/Plugins/Runtime/AudioSynesthesia/Source/AudioSynesthesiaCore/Public/PitchTracker.h
new file mode 100644
--- /dev/null
+++ b/Engine/Plugins/Runtime/AudioSynesthesia/Source/AudioSynesthesiaCore/Public/PitchTracker.h
@@ -0,0 +1,85 @@
+// Copyright Epic Games, Inc. All Rights Reserved.
+
+#pragma once
+
+#include "CoreMinimal.h"
+#include "DSP/BufferVectorOperations.h"
+
+namespace Audio
+{
+	/** Contains information about a pitch observation. */
+	struct FPitchInfo
+	{
+		/** The frequency (in Hz) of the observed pitch. */
+		float Frequency = 0.f;
+
+		/** The strength of the pitch.  The exact interpretation of this value will be
+		 * dependent upon the specific pitch detector. */
+		float Strength = 0.f;
+
+		/** The timestamp (in seconds) of the pitch observation. */
+		double Timestamp = 0.f;
+	};
+
+	/** Pitch Detector Interface 
+	 *
+	 * A pitch detector identifies pitch information from a stream of input audio.  The pitch detector should
+	 * support multiple calls to `DetectPitches(...)` originating from the same audio source. 
+	 */
+	class IPitchDetector
+	{
+		public:
+			virtual ~IPitchDetector() {}
+
+			/** Detect pitches from stream of audio. 
+			 *
+			 * @param InMonoAudio - The most recent input buffer from a stream of audio.
+			 * @param OutPitches - The pitches detected during this call to `DetectPitches` are appended to this array. 
+			 */
+			virtual void DetectPitches(const AlignedFloatBuffer& InMonoAudio, TArray<FPitchInfo>& OutPitches) = 0;
+
+
+			/** Mark the end of an audio stream. 
+			 *
+			 * Pitch detectors which need to know when an audio stream is finished can overload this method
+			 * to handle the end of an audio stream.
+			 *
+			 * @param OutPitches - The pitches detected during this call to `Finalize` are appended to this array. 
+			 */
+			virtual void Finalize(TArray<FPitchInfo>& OutPitches) = 0;
+	};
+
+	/** Contains information about a pitch track observation. */
+	struct FPitchTrackInfo
+	{
+		/** The array of observations associated with the pitch track. */
+		TArray<FPitchInfo> Observations;
+	};
+
+	/** Pitch Tracker Interface.
+	 *
+	 * A pitch tracker idenities multiple pitch tracks from a stream of input audio. The pitch tracker
+	 * should support multiple calls to `TrackPitches(...)` originating from the same audio source. 
+	 */
+	class IPitchTracker
+	{
+		public:
+			virtual ~IPitchTracker() {}
+
+			/** Detect pitche tracks from stream of audio. 
+			 *
+			 * @param InMonoAudio - The most recent input buffer from a stream of audio.
+			 * @param OutPitchTracks - The pitch tracks detected during this call to `TrackPitches` are appended to this array. 
+			 */
+			virtual void TrackPitches(const AlignedFloatBuffer& InMonoAudio, TArray<FPitchTrackInfo>& OutPitchTracks) = 0;
+
+			/** Mark the end of an audio stream. 
+			 *
+			 * Pitch trackers which need to know when an audio stream is finished can overload this method
+			 * to handle the end of an audio stream.
+			 *
+			 * @param OutPitchTracks - The pitch tracks detected during this call to `Finalize` are added to this array. 
+			 */
+			virtual void Finalize(TArray<FPitchTrackInfo>& OutPitchTracks) = 0;
+	};
+}
diff --git a/Engine/Source/Runtime/SignalProcessing/Private/BlockCorrelator.cpp b/Engine/Source/Runtime/SignalProcessing/Private/BlockCorrelator.cpp
new file mode 100644
--- /dev/null
+++ b/Engine/Source/Runtime/SignalProcessing/Private/BlockCorrelator.cpp
@@ -0,0 +1,221 @@
+// Copyright Epic Games, Inc. All Rights Reserved.
+
+#include "DSP/BlockCorrelator.h"
+#include "DSP/FFTAlgorithm.h"
+#include "DSP/FloatArrayMath.h"
+
+#include "SignalProcessingModule.h"
+
+namespace Audio
+{
+	FBlockCorrelator::FBlockCorrelator(const FBlockCorrelatorSettings& InSettings)
+	:	Settings(InSettings)
+	,	NumValuesInBlock(1 << Settings.Log2NumValuesInBlock)
+	,	NumValuesInFFTRealBuffer(0)
+	,	NumValuesInFFTComplexBuffer(0)
+	,	Window(InSettings.WindowType, 1 << Settings.Log2NumValuesInBlock, 1, false)
+	{
+		// Determine FFT Settings
+		FFFTSettings FFTSettings;
+
+		// Make FFT twice as big as block to avoid temporal aliasing.
+		FFTSettings.Log2Size = Settings.Log2NumValuesInBlock + 1;
+		FFTSettings.bArrays128BitAligned = true;
+		FFTSettings.bEnableHardwareAcceleration = true;
+
+		// Create FFT
+		FFTAlgorithm = FFFTFactory::NewFFTAlgorithm(FFTSettings);
+
+		if (!FFTAlgorithm.IsValid())
+		{
+			UE_LOG(LogSignalProcessing, Warning, TEXT("Failed to create FFT algorithm for FBlockCorrelator"));
+
+			// Set these to default value to avoid having to do many if statements during calls to *Correlate(...)
+			NumValuesInFFTRealBuffer = NumValuesInBlock * 2;
+			NumValuesInFFTComplexBuffer = NumValuesInBlock + 2;
+		}
+		else
+		{
+			NumValuesInFFTRealBuffer = FFTAlgorithm->NumInputFloats();
+			NumValuesInFFTComplexBuffer = FFTAlgorithm->NumOutputFloats();
+		}
+
+		// Allocate buffers
+		if (NumValuesInFFTRealBuffer > 0)
+		{
+			WindowedBufferA.AddZeroed(NumValuesInFFTRealBuffer);
+			WindowedBufferB.AddZeroed(NumValuesInFFTRealBuffer);
+			FullOutputBuffer.AddZeroed(NumValuesInFFTRealBuffer);
+		}
+
+		if (NumValuesInFFTComplexBuffer > 0)
+		{
+			ComplexBufferA.AddZeroed(NumValuesInFFTComplexBuffer);
+			ComplexBufferB.AddZeroed(NumValuesInFFTComplexBuffer);
+			ComplexCorrelationBuffer.AddZeroed(NumValuesInFFTComplexBuffer);
+		}
+
+		// Initialize normalization buffer to valid values. 
+		InitializeNormalizationBuffer();
+	}
+
+
+	const FBlockCorrelatorSettings& FBlockCorrelator::GetSettings() const
+	{
+		return Settings;
+	}
+
+	int32 FBlockCorrelator::GetNumInputValues() const
+	{
+		return NumValuesInBlock;
+	}
+
+	int32 FBlockCorrelator::GetNumOutputValues() const
+	{
+		return NumValuesInFFTRealBuffer;
+	}
+
+	void FBlockCorrelator::CrossCorrelate(const AlignedFloatBuffer& InputA, const AlignedFloatBuffer& InputB, AlignedFloatBuffer& Output)
+	{
+		check(InputA.Num() == NumValuesInBlock);
+		check(InputB.Num() == NumValuesInBlock);
+		check(Output.Num() == NumValuesInFFTRealBuffer);
+
+		// Be a little defensive about folks passing in the wrong size of input buffer.
+		int32 NumA = FMath::Min(NumValuesInBlock, InputA.Num());
+		int32 NumB = FMath::Min(NumValuesInBlock, InputB.Num());
+
+		if ((NumA > 0) && (NumB > 0))
+		{
+			// WindowBufferA and WindowBufferB are twice as long as the input and padded with zeros.
+			// This memcpy effectively zeropads the input buffers.
+			FMemory::Memcpy(WindowedBufferA.GetData(), InputA.GetData(), NumA * sizeof(float));
+			FMemory::Memcpy(WindowedBufferB.GetData(), InputB.GetData(), NumB * sizeof(float));
+
+			Window.ApplyToBuffer(WindowedBufferA.GetData());
+			Window.ApplyToBuffer(WindowedBufferB.GetData());
+
+			CyclicCrossCorrelate(WindowedBufferA, WindowedBufferB, Output);
+
+			if (Settings.bDoNormalize)
+			{
+				ArrayMultiplyInPlace(NormalizationBuffer, Output);
+			}
+		}
+	}
+
+	void FBlockCorrelator::AutoCorrelate(const AlignedFloatBuffer& Input, AlignedFloatBuffer& Output)
+	{
+		check(Input.Num() == NumValuesInBlock);
+		check(Output.Num() == NumValuesInFFTRealBuffer);
+
+		// Be a little defensive about folks passing in the wrong size of input buffer.
+		int32 Num = FMath::Min(NumValuesInBlock, Input.Num());
+
+		if (Num > 0)
+		{
+			// WindowBufferA is twice as long as the input and padded with zeros.
+			// This memcpy effectively zeropads the input buffer.
+			FMemory::Memcpy(WindowedBufferA.GetData(), Input.GetData(), Num * sizeof(float));
+
+			Window.ApplyToBuffer(WindowedBufferA.GetData());
+
+			CyclicAutoCorrelate(WindowedBufferA, Output);
+
+			if (Settings.bDoNormalize)
+			{
+				ArrayMultiplyInPlace(NormalizationBuffer, Output);
+			}
+		}
+	}
+
+	void FBlockCorrelator::CyclicCrossCorrelate(const AlignedFloatBuffer& InputA, const AlignedFloatBuffer& InputB, AlignedFloatBuffer& Output)
+	{
+		check(InputA.Num() == NumValuesInFFTRealBuffer);
+		check(InputB.Num() == NumValuesInFFTRealBuffer);
+		check(Output.Num() == NumValuesInFFTRealBuffer);
+
+		// Cyclical correlation can be done more quickly in the frequency domain. 
+		if (FFTAlgorithm.IsValid())
+		{
+			// Convert to freq domain.
+			FFTAlgorithm->ForwardRealToComplex(InputA.GetData(), ComplexBufferA.GetData());
+			FFTAlgorithm->ForwardRealToComplex(InputB.GetData(), ComplexBufferB.GetData());
+
+			// Do complex conjugate multiplication.
+			ArrayComplexConjugateInPlace(ComplexBufferB);
+			ArrayComplexMultiplyInPlace(ComplexBufferB, ComplexBufferA);
+
+			// Back to the time domain.
+			FFTAlgorithm->InverseComplexToReal(ComplexBufferA.GetData(), Output.GetData());
+		}
+	}
+
+	void FBlockCorrelator::CyclicAutoCorrelate(const AlignedFloatBuffer& Input, AlignedFloatBuffer& Output)
+	{
+		check(Input.Num() == NumValuesInFFTRealBuffer);
+		check(Output.Num() == NumValuesInFFTRealBuffer);
+
+		// Cyclical correlation can be done more quickly in the frequency domain. 
+		if (FFTAlgorithm.IsValid())
+		{
+			// Convert to freq domain.
+			FFTAlgorithm->ForwardRealToComplex(Input.GetData(), ComplexBufferA.GetData());
+
+			// Do complex conjugate multiplication.
+			ArrayComplexConjugate(ComplexBufferA, ComplexBufferB);
+			ArrayComplexMultiplyInPlace(ComplexBufferB, ComplexBufferA);
+
+			// Back to the time domain.
+			FFTAlgorithm->InverseComplexToReal(ComplexBufferA.GetData(), Output.GetData());
+		}
+	}
+
+	void FBlockCorrelator::InitializeNormalizationBuffer()
+	{
+		// Normalization by the autocorrelation of the window. This increases variance
+		// of information at end of autocorrelation output, but normalizes the output
+		// magnitude to remove bais.
+		if (!FFTAlgorithm.IsValid() || (NumValuesInBlock < 1))
+		{
+			return;
+		}
+
+		/** Determine the normalization buffer by taking the inverse of the autocorrelation of the 
+		 * input window.*/
+
+		NormalizationBuffer.Reset();
+
+		if ((NumValuesInFFTRealBuffer > 0) && (NumValuesInFFTRealBuffer > NumValuesInBlock))
+		{
+			// Initialize normalization buffer to be the time domain window applied to incoming data.
+			NormalizationBuffer.Init(1.f, NumValuesInBlock);
+			NormalizationBuffer.AddZeroed(NumValuesInFFTRealBuffer - NumValuesInBlock);
+
+			Window.ApplyToBuffer(NormalizationBuffer.GetData());
+		}
+
+		AlignedFloatBuffer CorrelationBuffer;
+		CorrelationBuffer.AddUninitialized(NumValuesInFFTRealBuffer);
+
+		// Take auto correlation of the window.
+		CyclicAutoCorrelate(NormalizationBuffer, CorrelationBuffer);
+
+		float* NormData = NormalizationBuffer.GetData();
+		const float* CorrData = CorrelationBuffer.GetData();
+
+		for (int32 i = 0; i < NumValuesInFFTRealBuffer; i++)
+		{
+			if (CorrData[i] > SMALL_NUMBER)
+			{
+				// Set the normalization data to be the inverse of the window's autocorrelation.
+				NormData[i] = 1.f / CorrData[i];
+			}
+			else
+			{
+				// Avoid divide by zero. 
+				NormData[i] = 0.f;
+			}
+		}
+	}
+}
diff --git a/Engine/Source/Runtime/SignalProcessing/Private/FloatArrayMath.cpp b/Engine/Source/Runtime/SignalProcessing/Private/FloatArrayMath.cpp
--- a/Engine/Source/Runtime/SignalProcessing/Private/FloatArrayMath.cpp
+++ b/Engine/Source/Runtime/SignalProcessing/Private/FloatArrayMath.cpp
@@ -3,7 +3,6 @@
 #include "DSP/FloatArrayMath.h"
 #include "CoreMinimal.h"
 
-
 namespace Audio
 {
 	namespace MathIntrinsics
@@ -352,6 +351,8 @@ namespace Audio
 
 	void ArrayMultiplyInPlace(TArrayView<const float> InValues1, TArrayView<float> InValues2)
 	{
+		check(InValues1.Num() == InValues2.Num());
+
 		const int32 Num = InValues1.Num();
 
 		const float* InData1 = InValues1.GetData();
@@ -388,6 +389,68 @@ namespace Audio
 		}
 	}
 
+	void ArrayComplexMultiplyInPlace(TArrayView<const float> InValues1, TArrayView<float> InValues2)
+	{
+		check(InValues1.Num() == InValues2.Num());
+
+		const int32 Num = InValues1.Num();
+
+		// Needs to be in interleaved format.
+		check((Num % 2) == 0);
+
+		const float* InData1 = InValues1.GetData();
+		float* InData2 = InValues2.GetData();
+
+		for (int32 i = 0; i < Num; i += 2)
+		{
+			float Real = (InData1[i] * InData2[i]) - (InData1[i + 1] * InData2[i + 1]);
+			float Imag = (InData1[i] * InData2[i + 1]) + (InData1[i + 1] * InData2[i]);
+			InData2[i] = Real;
+			InData2[i + 1] = Imag;
+		}
+	}
+
+	void ArrayComplexMultiplyInPlace(const AlignedFloatBuffer& InValues1, AlignedFloatBuffer& InValues2)
+	{
+		check(InValues1.Num() == InValues2.Num());
+
+		const int32 Num = InValues1.Num();
+		const int32 NumToSimd = Num & MathIntrinsics::SimdMask;
+		const int32 NumNotToSimd = Num & MathIntrinsics::NotSimdMask;
+
+		const float* InData1 = InValues1.GetData();
+		float* InData2 = InValues2.GetData();
+
+		if (NumToSimd)
+		{
+			const VectorRegister RealSignFlip = MakeVectorRegister(-1.f, 1.f, -1.f, 1.f);
+
+			for (int32 i = 0; i < NumToSimd; i += 4)
+			{
+				VectorRegister VectorData1 = VectorLoadAligned(&InData1[i]);
+				VectorRegister VectorData2 = VectorLoadAligned(&InData2[i]);
+
+				VectorRegister VectorData1Real = VectorSwizzle(VectorData1, 0, 0, 2, 2);
+				VectorRegister VectorData1Imag = VectorSwizzle(VectorData1, 1, 1, 3, 3);
+				VectorRegister VectorData2Swizzle = VectorSwizzle(VectorData2, 1, 0, 3, 2);
+
+				VectorRegister Result = VectorMultiply(VectorData1Imag, VectorData2Swizzle);
+				Result = VectorMultiply(Result, RealSignFlip);
+				Result = VectorMultiplyAdd(VectorData1Real, VectorData2, Result);
+
+				VectorStoreAligned(Result, &InData2[i]);
+			}
+		}
+
+		if (NumNotToSimd)
+		{
+			TArrayView<const float> ValuesView1(&InData1[NumToSimd],  NumNotToSimd);
+			TArrayView<float> ValuesView2(&InData2[NumToSimd],  NumNotToSimd);
+
+			ArrayComplexMultiplyInPlace(ValuesView1, ValuesView2);
+		}
+	}
+
 	void ArrayMultiplyByConstantInPlace(TArrayView<float> InValues, float InMultiplier)
 	{
 		const int32 Num = InValues.Num();
@@ -534,6 +597,95 @@ namespace Audio
 		}
 	}
 
+	void ArrayComplexConjugate(TArrayView<const float> InValues, TArrayView<float> OutValues)
+	{
+		check(OutValues.Num() == InValues.Num());
+		check((InValues.Num() % 2) == 0);
+
+		int32 Num = InValues.Num();
+
+		const float* InData = InValues.GetData();
+		float* OutData = OutValues.GetData();
+		
+		for (int32 i = 0; i < Num; i+= 2)
+		{
+			OutData[i] = InData[i];
+			OutData[i + 1] = -InData[i + 1];
+		}
+	}
+
+	void ArrayComplexConjugate(const AlignedFloatBuffer& InValues, AlignedFloatBuffer& OutValues)
+	{
+		check(OutValues.Num() == InValues.Num());
+
+		const int32 Num = InValues.Num();
+		const int32 NumToSimd = Num & MathIntrinsics::SimdMask;
+		const int32 NumNotToSimd = Num & MathIntrinsics::NotSimdMask;
+
+		const float* InData = InValues.GetData();
+		float* OutData = OutValues.GetData();
+
+		const VectorRegister ConjugateMult = MakeVectorRegister(1.f, -1.f, 1.f, -1.f);
+
+		for (int32 i = 0; i < NumToSimd; i += 4)
+		{
+			VectorRegister VectorData = VectorLoadAligned(&InData[i]);
+			
+			VectorData = VectorMultiply(VectorData, ConjugateMult);
+
+			VectorStoreAligned(VectorData, &OutData[i]);
+		}
+
+		if (NumNotToSimd)
+		{
+			TArrayView<const float> InView(&InData[NumToSimd], NumNotToSimd);
+			TArrayView<float> OutView(&OutData[NumToSimd], NumNotToSimd);
+
+			ArrayComplexConjugate(InView, OutView);
+		}
+	}
+
+	void ArrayComplexConjugateInPlace(TArrayView<float> InValues)
+	{
+		check((InValues.Num() % 2) == 0);
+
+		int32 Num = InValues.Num();
+
+		float* InData = InValues.GetData();
+		
+		for (int32 i = 1; i < Num; i+= 2)
+		{
+			InData[i] *= -1.f;
+		}
+	}
+
+	void ArrayComplexConjugateInPlace(AlignedFloatBuffer& InValues)
+	{
+		const int32 Num = InValues.Num();
+		const int32 NumToSimd = Num & MathIntrinsics::SimdMask;
+		const int32 NumNotToSimd = Num & MathIntrinsics::NotSimdMask;
+
+		float* InData = InValues.GetData();
+
+		const VectorRegister ConjugateMult = MakeVectorRegister(1.f, -1.f, 1.f, -1.f);
+
+		for (int32 i = 0; i < NumToSimd; i += 4)
+		{
+			VectorRegister VectorData = VectorLoadAligned(&InData[i]);
+			
+			VectorData = VectorMultiply(VectorData, ConjugateMult);
+
+			VectorStoreAligned(VectorData, &InData[i]);
+		}
+
+		if (NumNotToSimd)
+		{
+			TArrayView<float> InView(&InData[NumToSimd], NumNotToSimd);
+
+			ArrayComplexConjugateInPlace(InView);
+		}
+	}
+
 	void ArrayMagnitudeToDecibelInPlace(TArrayView<float> InValues, float InMinimumDb)
 	{
 		const int32 Num = InValues.Num();
diff --git a/Engine/Source/Runtime/SignalProcessing/Public/DSP/BlockCorrelator.h b/Engine/Source/Runtime/SignalProcessing/Public/DSP/BlockCorrelator.h
new file mode 100644
--- /dev/null
+++ b/Engine/Source/Runtime/SignalProcessing/Public/DSP/BlockCorrelator.h
@@ -0,0 +1,104 @@
+// Copyright Epic Games, Inc. All Rights Reserved.
+
+#pragma once
+
+#include "CoreMinimal.h"
+#include "DSP/BufferVectorOperations.h"
+#include "DSP/AudioFFT.h"
+#include "DSP/FFTAlgorithm.h"
+
+namespace Audio
+{	
+	/** Settings for block correlator */
+	struct FBlockCorrelatorSettings
+	{
+		/** Logarithm base 2 of the number of values in a block.*/
+		int32 Log2NumValuesInBlock = 10;  // 2^10 = 1024
+
+		/** Audio window type to apply to input. */
+		EWindowType WindowType = EWindowType::Blackman;
+
+		/** If true, autocorrelation values will be normalized to reduce bias introduced by block windowing. */
+		bool bDoNormalize = true;
+	};
+
+	/** Correlation calculator which utilizes FFT to perform fast correlation calculations. */
+	class SIGNALPROCESSING_API FBlockCorrelator
+	{
+		public:
+			/** Construct a block correlator with FBlockCorrelatorSettings */
+			FBlockCorrelator(const FBlockCorrelatorSettings& InSettings);
+
+			/** Return settings used to construct block correlator. */
+			const FBlockCorrelatorSettings& GetSettings() const;
+
+			/** Returns the number of input values to use when calling CrossCorrelate(...) or AutoCorrelate(...) */
+			int32 GetNumInputValues() const;
+
+			/** Returns the number of output values to use when calling CrossCorrelate(...) or AutoCorrelate(...) */
+			int32 GetNumOutputValues() const;
+
+			/** Cross correlate two input signals. 
+			 *
+			 * The Output array is filled with the cross correlation between InputA and InputB. Care should be taken 
+			 * when interpreting the output as the elements are ordered so that the cross correlation value of zero
+			 * lag is in the first element.
+			 *
+			 * CrossCorr[k] = Sum_n(InputA[n] * InputB[n + k])
+			 *
+			 * Then for M = `GetNumInputValues()` the Output array contains:
+			 * Output = [CrossCorr[0], CrossCorr[1], ..., CrossCorr[M-1], CrossCorr[M], CrossCorr[-M + 1], CrossCorr[-M + 2], ..., CrossCorr[-2], CrossCorr[-1]]
+			 *
+			 * @param InputA - First input block with `GetNumInputValues()` elements. 
+			 * @param InputB - Second input block with `GetNumInputValues()` elements. 
+			 * @param Output - Output block with `GetNumOutputValues()` elements.
+			 */
+			void CrossCorrelate(const AlignedFloatBuffer& InputA, const AlignedFloatBuffer& InputB, AlignedFloatBuffer& Output);
+
+			/** Autocorrelate an input signal. 
+			 *
+			 * The Output array is filled with the autocorrelation of the Input. Care should be taken when interpreting the 
+			 * output as the elements are ordered so that the autocorrelation value of zero lag is in the first element.
+			 *
+			 * AutoCorr[k] = Sum_n(Input[n] * Input[n + k])
+			 *
+			 * Then for M = `GetNumInputValues()` the Output array contains:
+			 * Output = [AutoCorr[0], AutoCorr[1], ..., AutoCorr[M-1], AutoCorr[M], AutoCorr[-M + 1], AutoCorr[-M + 2], ..., AutoCorr[-2], AutoCorr[-1]]
+			 *
+			 * @param Input - input block with `GetNumInputValues()` elements. 
+			 * @param Output - Output block with `GetNumOutputValues()` elements.
+			 */
+			void AutoCorrelate(const AlignedFloatBuffer& Input, AlignedFloatBuffer& Output);
+
+
+		private:
+			// No normalization or windowing.
+			void CyclicCrossCorrelate(const AlignedFloatBuffer& InputA, const AlignedFloatBuffer& InputB, AlignedFloatBuffer& Output);
+
+			// No normalization or windowing.
+			void CyclicAutoCorrelate(const AlignedFloatBuffer& Input, AlignedFloatBuffer& Output);
+
+			void InitializeNormalizationBuffer();
+
+			FBlockCorrelatorSettings Settings;
+
+			int32 NumValuesInBlock;
+			int32 NumValuesInFFTRealBuffer;
+			int32 NumValuesInFFTComplexBuffer;
+
+			TUniquePtr<IFFTAlgorithm> FFTAlgorithm;
+
+			FWindow Window;
+			AlignedFloatBuffer NormalizationBuffer;
+
+			AlignedFloatBuffer WindowedBufferA;
+			AlignedFloatBuffer WindowedBufferB;
+
+			AlignedFloatBuffer ComplexBufferA;
+			AlignedFloatBuffer ComplexBufferB;
+
+			AlignedFloatBuffer ComplexCorrelationBuffer;
+
+			AlignedFloatBuffer FullOutputBuffer;
+	};
+}
diff --git a/Engine/Source/Runtime/SignalProcessing/Public/DSP/Dsp.h b/Engine/Source/Runtime/SignalProcessing/Public/DSP/Dsp.h
--- a/Engine/Source/Runtime/SignalProcessing/Public/DSP/Dsp.h
+++ b/Engine/Source/Runtime/SignalProcessing/Public/DSP/Dsp.h
@@ -208,6 +208,38 @@ namespace Audio
 		return OutQ;
 	}
 
+	// Given three values, determine peak location and value of quadratic fitted to the data.
+	//
+	// @param InValues - An array of 3 values with the maximum value located in InValues[1].
+	// @param OutPeakLoc - The peak location relative to InValues[1].
+	// @param OutPeakValue - The value of the peak at the peak location.
+	//
+	// @returns True if a peak was found, false if the values do not represent a peak.
+	static FORCEINLINE bool QuadraticPeakInterpolation(const float InValues[3], float& OutPeakLoc, float& OutPeakValue)
+	{
+		float Denom = InValues[0] - 2.f * InValues[1] + InValues[2];
+
+		if (Denom >= 0.f)
+		{
+			// This is not a peak.
+			return false;
+		}
+
+		float Tmp = InValues[0] - InValues[2];
+
+		OutPeakLoc = 0.5f * Tmp / Denom;
+
+		if ((OutPeakLoc > 0.5f) || (OutPeakLoc < -0.5f))
+		{
+			// This is not a peak.
+			return false;
+		}
+
+		OutPeakValue = InValues[1] - 0.25f * Tmp * OutPeakLoc;
+
+		return true;
+	}
+
 	// Polynomial interpolation using lagrange polynomials. 
 	// https://en.wikipedia.org/wiki/Lagrange_polynomial
 	static FORCEINLINE float LagrangianInterpolation(const TArray<FVector2D> Points, const float Alpha)
diff --git a/Engine/Source/Runtime/SignalProcessing/Public/DSP/FloatArrayMath.h b/Engine/Source/Runtime/SignalProcessing/Public/DSP/FloatArrayMath.h
--- a/Engine/Source/Runtime/SignalProcessing/Public/DSP/FloatArrayMath.h
+++ b/Engine/Source/Runtime/SignalProcessing/Public/DSP/FloatArrayMath.h
@@ -89,6 +89,13 @@ namespace Audio
 	SIGNALPROCESSING_API void ArrayMultiplyInPlace(TArrayView<const float> InValues1, TArrayView<float> InValues2);
 	SIGNALPROCESSING_API void ArrayMultiplyInPlace(const AlignedFloatBuffer& InValues1, AlignedFloatBuffer& InValues2);
 
+	/** Multiplies two complex valued arrays element-wise. 
+	 * This assumes elements are in interleaved format [real_0, imag_0, ..., real_N, imag_N]
+	 * Stores result in InValues2
+	 */
+	SIGNALPROCESSING_API void ArrayComplexMultiplyInPlace(TArrayView<const float> InValues1, TArrayView<float> InValues2);
+	SIGNALPROCESSING_API void ArrayComplexMultiplyInPlace(const AlignedFloatBuffer& InValues1, AlignedFloatBuffer& InValues2);
+
 	/** Multiplies each element in InView by InMultiplier */
 	SIGNALPROCESSING_API void ArrayMultiplyByConstantInPlace(TArrayView<float> InValues, float InMultiplier);
 	SIGNALPROCESSING_API void ArrayMultiplyByConstantInPlace(AlignedFloatBuffer& InValues, float InMultiplier);
@@ -114,8 +121,17 @@ namespace Audio
 	SIGNALPROCESSING_API void ArraySubtractByConstantInPlace(AlignedFloatBuffer& InValues, float InSubtrahend);
 
 
+	/** Take Square Root of values in place. */
 	SIGNALPROCESSING_API void ArraySqrtInPlace(TArrayView<float> InValues);
 
+
+	/** Perform complex conjugate of array.  Assumes complex numbers are interlaves [real_0, imag_0, real_1, image_1, ..., real_N, imag_N]. */
+	SIGNALPROCESSING_API void ArrayComplexConjugate(TArrayView<const float> InValues, TArrayView<float> OutValues);
+	SIGNALPROCESSING_API void ArrayComplexConjugate(const AlignedFloatBuffer& InValues, AlignedFloatBuffer& OutValues);
+
+	SIGNALPROCESSING_API void ArrayComplexConjugateInPlace(TArrayView<float> InValues);
+	SIGNALPROCESSING_API void ArrayComplexConjugateInPlace(AlignedFloatBuffer& InValues);
+
 	/** Convert magnitude values to decibel values in place. db = 20 * log10(val) */
 	SIGNALPROCESSING_API void ArrayMagnitudeToDecibelInPlace(TArrayView<float> InValues, float InMinimumDb);
 	SIGNALPROCESSING_API void ArrayMagnitudeToDecibelInPlace(AlignedFloatBuffer& InValues, float InMinimumDb);
diff --git a/Engine/Source/Runtime/SignalProcessing/Public/DSP/SlidingWindow.h b/Engine/Source/Runtime/SignalProcessing/Public/DSP/SlidingWindow.h
--- a/Engine/Source/Runtime/SignalProcessing/Public/DSP/SlidingWindow.h
+++ b/Engine/Source/Runtime/SignalProcessing/Public/DSP/SlidingWindow.h
@@ -31,11 +31,6 @@ namespace Audio
 
 
 		public:
-			/** NumWindowSamples describes the number of samples in a window */
-			const int32 NumWindowSamples;
-
-			/** NumHopSamples describes the number of samples between adjacent windows */
-			const int32 NumHopSamples;
 
 			/**
 			 * Constructs a TSlidingBuffer with a constant window and hop size
@@ -49,6 +44,18 @@ namespace Audio
 				check(NumHopSamples >= 1);
 			}
 
+			/** Returns the number of samples in a window. */
+			int32 GetNumWindowSamples() const
+			{
+				return NumWindowSamples;
+			}
+
+			/** Returns the number of samples between windows. */
+			int32 GetNumHopSamples() const
+			{
+				return NumHopSamples;
+			}
+
 			/**
 			 * StoreForFutureWindows stores the necessary samples from InBuffer which will
 			 * be needed for future windows. It ignores all values in InBuffer which can
@@ -141,6 +148,11 @@ namespace Audio
 			}
 
 		private:
+			// NumWindowSamples describes the number of samples in a window.
+			int32 NumWindowSamples;
+
+			// NumHopSamples describes the number of samples between adjacent windows.
+			int32 NumHopSamples;
 
 			// Stores samples from previous calls which are still needed for future buffers
 			TArray<InSampleType> StorageBuffer;
