Clean-up InterpolatedOnePole & expose target value in param interpolator
#rnx

diff --git a/Engine/Source/Runtime/SignalProcessing/Private/InterpolatedOnePole.cpp b/Engine/Source/Runtime/SignalProcessing/Private/InterpolatedOnePole.cpp
--- a/Engine/Source/Runtime/SignalProcessing/Private/InterpolatedOnePole.cpp
+++ b/Engine/Source/Runtime/SignalProcessing/Private/InterpolatedOnePole.cpp
@@ -2,17 +2,18 @@
 
 #include "DSP/InterpolatedOnePole.h"
 
+#include "AudioDefines.h"
+
+
 namespace Audio
 {
 
 	// INTERPOLATED ONE-POLE LOW-PASS IMPLEMENTATION
 	FInterpolatedLPF::FInterpolatedLPF()
 	{
-		Z1.Init(0.0f, NumChannels);
 		Reset();
 	}
 
-
 	void FInterpolatedLPF::Init(float InSampleRate, int32 InNumChannels)
 	{
 		SampleRate = InSampleRate;
@@ -23,20 +24,21 @@ namespace Audio
 		Reset();
 	}
 
-	void FInterpolatedLPF::StartFrequencyInterpolation(const float InTargetFrequency, const int32 InterpLength)
+	void FInterpolatedLPF::StartFrequencyInterpolation(const float InTargetFrequency, const int32 InInterpLength)
 	{
-		CurrInterpLength = InterpLength;
+		CurrInterpLength = InInterpLength;
 
-		if(isFirstFrequencyChange)
+		if (bIsFirstFrequencyChange)
 		{
 			CurrInterpLength = 0;
-			isFirstFrequencyChange = false;
+			bIsFirstFrequencyChange = false;
 		}
 
 		if (!FMath::IsNearlyEqual(InTargetFrequency, CutoffFrequency))
 		{
 			CutoffFrequency = InTargetFrequency;
-			float NormalizedFreq = FMath::Clamp(2.0f * InTargetFrequency / SampleRate, 0.0f, 1.0f);
+
+			const float NormalizedFreq = FMath::Clamp(2.0f * GetCutoffFrequency() / SampleRate, 0.0f, 1.0f);
 			B1Target = FMath::Exp(-PI * NormalizedFreq);
 			B1Delta = (B1Target - B1Curr) / static_cast<float>(CurrInterpLength);
 		}
@@ -100,7 +102,7 @@ namespace Audio
 		CurrInterpLength = 0;
 		ClearMemory();
 		Z1Data = Z1.GetData();
-		isFirstFrequencyChange = true;
+		bIsFirstFrequencyChange = true;
 	}
 
 	void FInterpolatedLPF::ClearMemory()
@@ -113,7 +115,6 @@ namespace Audio
 	// INTERPOLATED ONE-POLE HIGH-PASS IMPLEMENTATION
 	FInterpolatedHPF::FInterpolatedHPF()
 	{
-		Z1.Init(0.0f, NumChannels);
 		Reset();
 	}
 
@@ -132,19 +133,23 @@ namespace Audio
 	{
 		CurrInterpLength = InterpLength;
 
-		if (isFirstFrequencyChange)
+		if (bIsFirstFrequencyChange)
 		{
 			CurrInterpLength = 0;
-			isFirstFrequencyChange = false;
+			bIsFirstFrequencyChange = false;
 		}
 
 		if (!FMath::IsNearlyEqual(InTargetFrequency, CutoffFrequency))
 		{
 			CutoffFrequency = FMath::Min(InTargetFrequency, NyquistLimit);
 
-			const float G = GetGCoefficient();
-			A0Target = G / (1.0f + G);
+			// G computation is a reduced form of the following set of equations:
+			// OmegaDigital = 2.0f * PI * CutoffFrequency;
+			// OmegaAnalog = 2.0f * SampleRate * Audio::FastTan(0.5f * OmegaDigital / SampleRate);
+			// G = 0.5f * OmegaAnalog / SampleRate;
+			const float G = Audio::FastTan(PI * GetCutoffFrequency() / SampleRate);
 
+			A0Target = G / (1.0f + G);
 			A0Delta = (A0Target - A0Curr) / static_cast<float>(CurrInterpLength);
 		}
 
@@ -189,7 +194,6 @@ namespace Audio
 		}
 	}
 
-
 	void FInterpolatedHPF::Reset()
 	{
 		A0Curr = 0.0f;
@@ -197,7 +201,7 @@ namespace Audio
 		CurrInterpLength = 0;
 		ClearMemory();
 		Z1Data = Z1.GetData();
-		isFirstFrequencyChange = true;
+		bIsFirstFrequencyChange = true;
 	}
 
 	void FInterpolatedHPF::ClearMemory()
diff --git a/Engine/Source/Runtime/SignalProcessing/Public/DSP/InterpolatedOnePole.h b/Engine/Source/Runtime/SignalProcessing/Public/DSP/InterpolatedOnePole.h
--- a/Engine/Source/Runtime/SignalProcessing/Public/DSP/InterpolatedOnePole.h
+++ b/Engine/Source/Runtime/SignalProcessing/Public/DSP/InterpolatedOnePole.h
@@ -66,7 +66,7 @@ namespace Audio
 		int32 NumInterpSteps;
 		float SampleRate{ 0 };
 		int32 NumChannels{ 1 };
-		int32 isFirstFrequencyChange : 1;
+		bool bIsFirstFrequencyChange{ true };
 
 	}; // class Interpolated Low pass filter
 
@@ -121,14 +121,6 @@ namespace Audio
 		void ClearMemory();
 		
 	private:
-		FORCEINLINE float GetGCoefficient() const
-		{
-			const float OmegaDigital = 2.0f * PI * CutoffFrequency;
-			const float OmegaAnalog = 2.0f * SampleRate * Audio::FastTan(0.5f * OmegaDigital / SampleRate);
-			const float G = 0.5f * OmegaAnalog / SampleRate;
-			return G;
-		}
-
 		float CutoffFrequency{ -1.0f };
 		float A0Curr{ 0.0f }; // coefficient
 		float A0Delta{ 0.0f }; // coefficient step size
@@ -140,7 +132,7 @@ namespace Audio
 		float SampleRate{ 0 };
 		float NyquistLimit{ 0.0f };
 		int32 NumChannels{ 1 };
-		int32 isFirstFrequencyChange : 1;
+		bool bIsFirstFrequencyChange{ true };
 	
 	}; // class Interpolated High pass filter
 } // namespace Audio
diff --git a/Engine/Source/Runtime/SignalProcessing/Public/DSP/ParamInterpolator.h b/Engine/Source/Runtime/SignalProcessing/Public/DSP/ParamInterpolator.h
--- a/Engine/Source/Runtime/SignalProcessing/Public/DSP/ParamInterpolator.h
+++ b/Engine/Source/Runtime/SignalProcessing/Public/DSP/ParamInterpolator.h
@@ -60,6 +60,12 @@ namespace Audio
 			return CurrentValue;
 		}
 
+		// Returns the target value
+		FORCEINLINE float GetTarget() const
+		{
+			return TargetValue;
+		}
+
 	private:
 		float CurrentValue;
 		float StartingValue;
