Adding essential metasound directed graph algorithms.

diff --git a/Engine/Plugins/Runtime/Metasound/Source/MetasoundGraphCore/Private/MetasoundGraphAlgo.cpp b/Engine/Plugins/Runtime/Metasound/Source/MetasoundGraphCore/Private/MetasoundGraphAlgo.cpp
new file mode 100644
--- /dev/null
+++ b/Engine/Plugins/Runtime/Metasound/Source/MetasoundGraphCore/Private/MetasoundGraphAlgo.cpp
@@ -0,0 +1,399 @@
+// Copyright Epic Games, Inc. All Rights Reserved.
+
+#include "MetasoundGraphAlgo.h"
+
+#include "DSP/DirectedGraphAlgo.h"
+
+namespace Metasound
+{
+	// The core graph algorithms utilize integers to represent vertices
+	// and pairs of vertices to represent edges. This class provides 
+	// convenience functions for converting back and forth between 
+	// metasound types and the core graph algorithm types.
+	class FDirectedGraphAlgoAdapter
+	{
+		using FNodePair = TTuple<const INode*, const INode*>;
+		using FDirectedEdge = Audio::FDirectedEdge;
+
+		// The TLazyObject holds the object and a flag to denote whether the 
+		// object has been initialized.
+		template<typename T>
+		struct TLazyObject
+		{
+			bool bIsInitialized = false;
+
+			T& operator*() { return Object; }
+			const T& operator*() const { return Object; }
+
+			T* operator->() { return &Object; }
+			const T* operator->() const { return &Object; }
+
+			private:
+			T Object;
+		};
+
+		public:
+			FDirectedGraphAlgoAdapter(const IGraph& InGraph)
+			{
+				// Gather all nodes references in graph edges.
+				for (const FDataEdge& Edge : InGraph.GetDataEdges())
+				{
+					if (nullptr != Edge.From.Node) 
+					{
+						MetasoundNodeSet.Add(Edge.From.Node);	
+					}
+					
+					if (nullptr != Edge.To.Node)
+					{
+						MetasoundNodeSet.Add(Edge.To.Node);	
+					}
+				}
+
+				// Gather all nodes referenced as input destinations
+				for (const FInputDataDestinationCollection::ElementType& Dest : InGraph.GetInputDataDestinations())
+				{
+					MetasoundNodeSet.Add(Dest.Value.Node);
+				}
+
+				// Gather all nodes referenced as output sources.
+				for (const FOutputDataSourceCollection::ElementType& Source : InGraph.GetOutputDataSources())
+				{
+					MetasoundNodeSet.Add(Source.Value.Node);
+				}
+
+				// Convert to array so we can use index of node to 
+				// represent the int32 needed for the core graph algorithms. 
+				MetasoundUniqueNodes = MetasoundNodeSet.Array();
+
+				for (int32 i = 0; i < MetasoundUniqueNodes.Num(); i++)
+				{
+					MetasoundNodeToInt.Add(MetasoundUniqueNodes[i], i);
+				}
+
+				for (const FDataEdge& Edge : InGraph.GetDataEdges())
+				{
+					if ((nullptr == Edge.From.Node) || (nullptr == Edge.To.Node))
+					{
+						// skip edges with invalid null nodes. 
+						continue;
+					}
+
+					int32 FromVertex = MetasoundNodeToInt[Edge.From.Node];
+					int32 ToVertex = MetasoundNodeToInt[Edge.To.Node];
+
+					EdgeSet.Add(FDirectedEdge(FromVertex, ToVertex));
+
+					MetasoundNodePairEdges.Emplace(FNodePair(Edge.From.Node, Edge.To.Node), Edge);
+				}
+
+				// Gather all vertices referenced as input destinations
+				for (const FInputDataDestinationCollection::ElementType& Dest : InGraph.GetInputDataDestinations())
+				{
+					InputVertices.AddUnique(MetasoundNodeToInt[Dest.Value.Node]);
+				}
+
+				// Gather all vertices referenced as output sources.
+				for (const FOutputDataSourceCollection::ElementType& Source : InGraph.GetOutputDataSources())
+				{
+					OutputVertices.AddUnique(MetasoundNodeToInt[Source.Value.Node]);
+				}
+			}
+
+			// All unique vertices in the graph.
+			const TArray<int32>& GetUniqueVertices() const
+			{
+				if (!Lazy.UniqueVertices.bIsInitialized)
+				{
+					Lazy.UniqueVertices->Reserve(MetasoundUniqueNodes.Num());
+
+					for (int32 i = 0; i < MetasoundUniqueNodes.Num(); i++)
+					{
+						Lazy.UniqueVertices->Add(i);
+					}
+
+					Lazy.UniqueVertices.bIsInitialized = true;
+				}
+
+				return *Lazy.UniqueVertices;
+			}
+
+			// All unique input vertices in the graph.
+			const TArray<int32>& GetUniqueInputVertices() const
+			{
+				return InputVertices;
+			}
+
+			// All unique output vertices in the graph.
+			const TArray<int32>& GetUniqueOutputVertices() const
+			{
+				return OutputVertices;
+			}
+
+			// Get the vertex representing the node.
+			int32 GetVertex(const INode* InNode) const
+			{
+				return MetasoundNodeToInt[InNode];
+			}
+
+			// Get all unique edges in the graph.
+			const TArray<FDirectedEdge>& GetUniqueEdges() const
+			{
+				if (!Lazy.UniqueEdges.bIsInitialized)
+				{
+					*Lazy.UniqueEdges = EdgeSet.Array();
+					Lazy.UniqueEdges.bIsInitialized = true;
+				}
+
+				return *Lazy.UniqueEdges;
+			}
+
+			// Get all unique edges in the graph.
+			const TSet<FDirectedEdge>& GetEdgeSet() const
+			{
+				return EdgeSet;
+			}
+
+			// Get node for vertex.
+			const INode* GetNode(int32 InVertex) const
+			{
+				return MetasoundUniqueNodes[InVertex];
+			}
+
+			// Get a set of nodes from a set of vertices.
+			void GetNodeSet(const TSet<int32>& InVertexSet, TSet<const INode*>& OutNodeSet) const
+			{
+				for (int32 Vertex : InVertexSet)
+				{
+					OutNodeSet.Add(MetasoundUniqueNodes[Vertex]);
+				}
+			}
+
+			// Get all the metasound edges represented by an edge.
+			void GetDataEdges(const FDirectedEdge& InEdge, TArray<FDataEdge>& OutMetasoundDataEdges) const
+			{
+				const INode* FromNode = MetasoundUniqueNodes[InEdge.Get<0>()];
+				const INode* ToNode = MetasoundUniqueNodes[InEdge.Get<1>()];
+
+				MetasoundNodePairEdges.MultiFind(FNodePair(FromNode, ToNode), OutMetasoundDataEdges);
+			}
+
+			// Get in the form of a directed tree
+			const Audio::FDirectedTree& GetTree() const
+			{
+				if (!Lazy.Tree.bIsInitialized)
+				{
+					Audio::FDirectedGraphAlgo::BuildDirectedTree(GetUniqueEdges(), *Lazy.Tree);
+					Lazy.Tree.bIsInitialized = true;
+				}
+
+				return *Lazy.Tree;
+			}
+
+			// Get in the form of a transpose directed tree
+			const Audio::FDirectedTree& GetTransposeTree() const
+			{
+				if (!Lazy.TransposeTree.bIsInitialized)
+				{
+					Audio::FDirectedGraphAlgo::BuildTransposeDirectedTree(GetUniqueEdges(), *Lazy.TransposeTree);
+					Lazy.TransposeTree.bIsInitialized = true;
+				}
+
+				return *Lazy.TransposeTree;
+			}
+			
+		private:
+			// Everything in the lazy cache is mutable.
+			struct FLazyCache
+			{
+				TLazyObject<TArray<int32>> UniqueVertices;
+				TLazyObject<TArray<Audio::FDirectedEdge>> UniqueEdges;
+				TLazyObject<Audio::FDirectedTree> Tree;
+				TLazyObject<Audio::FDirectedTree> TransposeTree;
+			};
+
+			TSet<FDirectedEdge> EdgeSet;
+
+			TMap<const INode*, int32> MetasoundNodeToInt;
+			TArray<int32> InputVertices;
+			TArray<int32> OutputVertices;
+
+
+			TArray<const INode*> MetasoundUniqueNodes;
+			TSet<const INode*> MetasoundNodeSet;
+
+			TMultiMap<FNodePair, FDataEdge> MetasoundNodePairEdges;
+
+			// Everything in the lazy cache is mutable.
+			mutable FLazyCache Lazy;
+	};
+
+	TPimplPtr<FDirectedGraphAlgoAdapter> FDirectedGraphAlgo::CreateDirectedGraphAlgoAdapter(const IGraph& InGraph)
+	{
+		return MakePimpl<FDirectedGraphAlgoAdapter>(InGraph);
+	}
+
+	bool FDirectedGraphAlgo::DepthFirstTopologicalSort(const IGraph& InGraph, TArray<const INode*>& OutNodeOrder)
+	{
+		return DepthFirstTopologicalSort(FDirectedGraphAlgoAdapter(InGraph), OutNodeOrder);
+	}
+
+	bool FDirectedGraphAlgo::DepthFirstTopologicalSort(const FDirectedGraphAlgoAdapter& InAdapter, TArray<const INode*>& OutNodeOrder)
+	{
+		using namespace Audio;
+
+		TArray<int32> VertexOrder;
+
+		// Call algo implementation
+		bool bResult = Audio::FDirectedGraphAlgo::DepthFirstTopologicalSort(InAdapter.GetUniqueVertices(), InAdapter.GetUniqueEdges(), VertexOrder);
+
+		if (!bResult)
+		{
+			return false;
+		}
+
+		// Convert back to Metasound types.
+		OutNodeOrder.Reserve(OutNodeOrder.Num() + VertexOrder.Num());
+
+		for (int32 Vertex : VertexOrder)
+		{
+			OutNodeOrder.Add(InAdapter.GetNode(Vertex));
+		}
+
+		return bResult;
+	}
+
+	bool FDirectedGraphAlgo::KahnTopologicalSort(const IGraph& InGraph, TArray<const INode*>& OutNodeOrder)
+	{
+		return KahnTopologicalSort(FDirectedGraphAlgoAdapter(InGraph), OutNodeOrder);
+	}
+
+	bool FDirectedGraphAlgo::KahnTopologicalSort(const FDirectedGraphAlgoAdapter& InAdapter, TArray<const INode*>& OutNodeOrder)
+	{
+		using namespace Audio;
+
+		TArray<int32> VertexOrder;
+
+		// Call algo implementation
+		bool bResult = Audio::FDirectedGraphAlgo::KahnTopologicalSort(InAdapter.GetUniqueVertices(), InAdapter.GetUniqueEdges(), VertexOrder);
+
+		if (!bResult)
+		{
+			return false;
+		}
+
+		// Convert back to Metasound types.
+		OutNodeOrder.Reserve(OutNodeOrder.Num() + VertexOrder.Num());
+
+		for (int32 Vertex : VertexOrder)
+		{
+			OutNodeOrder.Add(InAdapter.GetNode(Vertex));
+		}
+
+		return bResult;
+	}
+
+	bool FDirectedGraphAlgo::TarjanStronglyConnectedComponents(const IGraph& InGraph, TArray<FStronglyConnectedComponent>& OutComponents, bool bExcludeSingleVertex)
+	{
+		return TarjanStronglyConnectedComponents(FDirectedGraphAlgoAdapter(InGraph), OutComponents, bExcludeSingleVertex);
+	}
+
+	bool FDirectedGraphAlgo::TarjanStronglyConnectedComponents(const FDirectedGraphAlgoAdapter& InAdapter, TArray<FStronglyConnectedComponent>& OutComponents, bool bExcludeSingleVertex)
+	{
+		TArray<Audio::FStronglyConnectedComponent> StronglyConnectedComponents;
+
+		// Run tarjan on metasound derived graph edges 
+		if (Audio::FDirectedGraphAlgo::TarjanStronglyConnectedComponents(InAdapter.GetEdgeSet(), StronglyConnectedComponents, bExcludeSingleVertex))
+		{
+			// If strongly connected components are found, they must be converted
+			// back into metasound types. 
+			for (const Audio::FStronglyConnectedComponent& Component : StronglyConnectedComponents)
+			{
+				FStronglyConnectedComponent& MetasoundGraphComponent = OutComponents.AddDefaulted_GetRef();
+
+				for (int32 Vertex : Component.Vertices)
+				{
+					MetasoundGraphComponent.Nodes.Add(InAdapter.GetNode(Vertex));
+				}
+
+				for (const Audio::FDirectedEdge& Edge : Component.Edges)
+				{
+					TArray<FDataEdge> MetasoundComponentEdges;
+
+					InAdapter.GetDataEdges(Edge, MetasoundComponentEdges);
+
+					MetasoundGraphComponent.Edges.Append(MetasoundComponentEdges);
+				}
+			}
+
+			return true;
+		}
+
+		// No strongly connected components found.
+		return false;
+	}
+
+
+	void FDirectedGraphAlgo::FindReachableNodesFromInput(const IGraph& InGraph, TSet<const INode*>& OutNodes)
+	{
+		FindReachableNodesFromInput(FDirectedGraphAlgoAdapter(InGraph), OutNodes);
+	}
+
+	void FDirectedGraphAlgo::FindReachableNodesFromInput(const FDirectedGraphAlgoAdapter& InAdapter, TSet<const INode*>& OutNodes)
+	{
+		TSet<int32> VisitedVertices;
+
+		for (int32 Vertex : InAdapter.GetUniqueInputVertices())
+		{
+			Audio::FDirectedGraphAlgo::DepthFirstTraversal(Vertex, InAdapter.GetTree(), [&](int32 InVertexBeingVisited) -> bool 
+				{
+					bool bIsAlreadyInSet = false;
+					VisitedVertices.Add(InVertexBeingVisited, &bIsAlreadyInSet);
+
+					return !bIsAlreadyInSet;
+				}
+			);
+		}
+
+		InAdapter.GetNodeSet(VisitedVertices, OutNodes);
+	}
+
+	void FDirectedGraphAlgo::FindReachableNodesFromOutput(const IGraph& InGraph, TSet<const INode*>& OutNodes)
+	{
+		FindReachableNodesFromOutput(FDirectedGraphAlgoAdapter(InGraph), OutNodes);
+	}
+
+	void FDirectedGraphAlgo::FindReachableNodesFromOutput(const FDirectedGraphAlgoAdapter& InAdapter, TSet<const INode*>& OutNodes)
+	{
+		TSet<int32> VisitedVertices;
+
+		for (int32 Vertex : InAdapter.GetUniqueOutputVertices())
+		{ 
+			Audio::FDirectedGraphAlgo::DepthFirstTraversal(Vertex, InAdapter.GetTransposeTree(), [&](int32 InVertexBeingVisited) -> bool 
+				{
+					bool bIsAlreadyInSet = false;
+					VisitedVertices.Add(InVertexBeingVisited, &bIsAlreadyInSet);
+
+					return !bIsAlreadyInSet;
+				}
+			);
+		}
+
+		InAdapter.GetNodeSet(VisitedVertices, OutNodes);
+	}
+
+	void FDirectedGraphAlgo::FindReachableNodes(const IGraph& InGraph, TSet<const INode*>& OutNodes)
+	{
+		FindReachableNodes(FDirectedGraphAlgoAdapter(InGraph), OutNodes);
+	}
+
+	void FDirectedGraphAlgo::FindReachableNodes(const FDirectedGraphAlgoAdapter& InAdapter, TSet<const INode*>& OutNodes)
+	{
+		TSet<const INode*> NodesFromInput;
+		FindReachableNodesFromInput(InAdapter, NodesFromInput);
+
+		TSet<const INode*> NodesFromOutput;
+		FindReachableNodesFromOutput(InAdapter, NodesFromOutput);
+
+		OutNodes = NodesFromInput.Union(NodesFromOutput);
+	}
+}
diff --git a/Engine/Plugins/Runtime/Metasound/Source/MetasoundGraphCore/Private/MetasoundNodeInterface.cpp b/Engine/Plugins/Runtime/Metasound/Source/MetasoundGraphCore/Private/MetasoundNodeInterface.cpp
--- a/Engine/Plugins/Runtime/Metasound/Source/MetasoundGraphCore/Private/MetasoundNodeInterface.cpp
+++ b/Engine/Plugins/Runtime/Metasound/Source/MetasoundGraphCore/Private/MetasoundNodeInterface.cpp
@@ -18,4 +18,21 @@ namespace Metasound
 	{
 		return !(LHS == RHS);
 	}
+
+	bool operator==(const FOutputDataSource& InLeft, const FOutputDataSource& InRight)
+	{
+		return (InLeft.Node == InRight.Node) && (InLeft.Vertex == InRight.Vertex);
+	}
+
+	/** Check if two FInputDataDestinations are equal. */
+	bool operator==(const FInputDataDestination& InLeft, const FInputDataDestination& InRight)
+	{
+		return (InLeft.Node == InRight.Node) && (InLeft.Vertex == InRight.Vertex);
+	}
+
+	/** Check if two FDataEdges are equal. */
+	bool operator==(const FDataEdge& InLeft, const FDataEdge& InRight)
+	{
+		return (InLeft.From == InRight.From) && (InLeft.To == InRight.To);
+	}
 }
diff --git a/Engine/Plugins/Runtime/Metasound/Source/MetasoundGraphCore/Private/MetasoundOperatorBuilder.cpp b/Engine/Plugins/Runtime/Metasound/Source/MetasoundGraphCore/Private/MetasoundOperatorBuilder.cpp
--- a/Engine/Plugins/Runtime/Metasound/Source/MetasoundGraphCore/Private/MetasoundOperatorBuilder.cpp
+++ b/Engine/Plugins/Runtime/Metasound/Source/MetasoundGraphCore/Private/MetasoundOperatorBuilder.cpp
@@ -1,12 +1,13 @@
 // Copyright Epic Games, Inc. All Rights Reserved.
 
 #include "MetasoundOperatorBuilder.h"
+
 #include "MetasoundBuilderInterface.h"
-#include "MetasoundOperatorInterface.h"
+#include "MetasoundBuildError.h"
+#include "MetasoundGraphAlgo.h"
 #include "MetasoundGraphOperator.h"
-#include "MetasoundTarjan.h"
+#include "MetasoundOperatorInterface.h"
 
-#include "MetasoundBuildError.h"
 
 namespace Metasound
 {
@@ -140,11 +141,11 @@ namespace Metasound
 				bool bExcludeSingleVertex = true;
 
 				// Try to find cycles
-				TArray<FMetasoundStronglyConnectedComponent> Cycles;
+				TArray<FStronglyConnectedComponent> Cycles;
 
-				if(FTarjan::FindStronglyConnectedComponents(InGraph, Cycles, bExcludeSingleVertex))
+				if(FDirectedGraphAlgo::TarjanStronglyConnectedComponents(InGraph, Cycles, bExcludeSingleVertex))
 				{
-					for (const FMetasoundStronglyConnectedComponent& Cycle : Cycles)
+					for (const FStronglyConnectedComponent& Cycle : Cycles)
 					{
 						AddBuildError<FGraphCycleError>(OutErrors, Cycle.Nodes, Cycle.Edges);
 					}
diff --git a/Engine/Plugins/Runtime/Metasound/Source/MetasoundGraphCore/Private/MetasoundTarjan.cpp b/Engine/Plugins/Runtime/Metasound/Source/MetasoundGraphCore/Private/MetasoundTarjan.cpp
deleted file mode 100644
diff --git a/Engine/Plugins/Runtime/Metasound/Source/MetasoundGraphCore/Public/MetasoundGraphAlgo.h b/Engine/Plugins/Runtime/Metasound/Source/MetasoundGraphCore/Public/MetasoundGraphAlgo.h
new file mode 100644
--- /dev/null
+++ b/Engine/Plugins/Runtime/Metasound/Source/MetasoundGraphCore/Public/MetasoundGraphAlgo.h
@@ -0,0 +1,162 @@
+// Copyright Epic Games, Inc. All Rights Reserved.
+
+#pragma once
+
+#include "CoreMinimal.h"
+#include "MetasoundNodeInterface.h"
+#include "Templates/PimplPtr.h"
+
+namespace Metasound
+{
+	/** A strongly connected component containing Metasound INodes and 
+	 * FDataEdges.
+	 */
+	struct FStronglyConnectedComponent
+	{
+		/** Nodes in the strongly connected component. */
+		TArray<const INode*> Nodes;
+
+		/** FDataEdges in the strongly connected component. */
+		TArray<FDataEdge> Edges;
+	};
+
+	// Forward declare adapter which caches transformations on an IGraph and 
+	// for graph algorithm optimizations.
+	class FDirectedGraphAlgoAdapter;
+
+	struct METASOUNDGRAPHCORE_API FDirectedGraphAlgo
+	{
+		/** The FDirectedGraphAlgoAdapter caches internal representations of the
+		 * passed in IGraph. Using this can reduce redundant calcuations when 
+		 * multiple algorithms are run on the same IGraph (without altering the IGraph).
+		 *
+		 * This function creates the adapter. 
+		 */
+		static TPimplPtr<FDirectedGraphAlgoAdapter> CreateDirectedGraphAlgoAdapter(const IGraph& InGraph);
+
+		/** Sort the nodes of a directed acyclic graph using a depth-first 
+		 * topological sorting algorithm.
+		 *
+		 * @param InGraph      - The graph containing edges and nodes to sort. 
+		 * @param OutNodeOrder - An array of nodes from InGraph put in a valid 
+		 *                       topological sorted order.
+		 *
+		 * @return True if sorting was successful, false otherwise. 
+		 */
+		static bool DepthFirstTopologicalSort(const IGraph& InGraph, TArray<const INode*>& OutNodeOrder);
+
+		/** Sort the nodes of a directed acyclic graph using a depth-first 
+		 * topological sorting algorithm.
+		 *
+		 * @param InAdapter    - The adapter containing edges and nodes to sort. 
+		 * @param OutNodeOrder - An array of nodes from InGraph put in a valid 
+		 *                       topological sorted order.
+		 *
+		 * @return True if sorting was successful, false otherwise. 
+		 */
+		static bool DepthFirstTopologicalSort(const FDirectedGraphAlgoAdapter& InAdapter, TArray<const INode*>& OutNodeOrder);
+
+		/** Sort the nodes of a directed acyclic graph using Kahn's topological
+		 * sorting algorithm.
+		 *
+		 * @param InGraph      - The graph containing edges and nodes to sort. 
+		 * @param OutNodeOrder - An array of nodes from InGraph put in a valid 
+		 *                       topological sorted order.
+		 *
+		 * @return True if sorting was successful, false otherwise. 
+		 */
+		static bool KahnTopologicalSort(const IGraph& InGraph, TArray<const INode*>& OutNodeOrder);
+
+		/** Sort the nodes of a directed acyclic graph using Kahn's topological
+		 * sorting algorithm.
+		 *
+		 * @param InAdapter    - The adapter containing edges and nodes to sort. 
+		 * @param OutNodeOrder - An array of nodes from InGraph put in a valid 
+		 *                       topological sorted order.
+		 *
+		 * @return True if sorting was successful, false otherwise. 
+		 */
+		static bool KahnTopologicalSort(const FDirectedGraphAlgoAdapter& InAdapter, TArray<const INode*>& OutNodeOrder);
+
+		/** Find strongly connected components given a Metasound IGraph.
+		 *
+		 * @param InGraph              - Graph to be analyzed. 
+		 * @param OutComponents        - Strongly connected components found in
+		 * 								 the graph are be added to this array.
+		 * @param bExcludeSingleVertex - If true, single vertices are not be 
+		 *                               returned as strongly connected components. 
+		 *                               If false, single vertices may be returned
+		 *                               as strongly connected components. 
+		 *
+		 * @return True if one or more strongly connected components are added
+		 *         to OutComponents. False otherwise. 
+		 */
+		static bool TarjanStronglyConnectedComponents(const IGraph& InGraph, TArray<FStronglyConnectedComponent>& OutComponents, bool bExcludeSingleVertex = true);
+
+		/** Find strongly connected components given a Metasound IGraph.
+		 *
+		 * @param InAdapter            - The adapter containing edges and nodes to analyze.
+		 * @param OutComponents        - Strongly connected components found in
+		 * 								 the graph are be added to this array.
+		 * @param bExcludeSingleVertex - If true, single vertices are not be 
+		 *                               returned as strongly connected components. 
+		 *                               If false, single vertices may be returned
+		 *                               as strongly connected components. 
+		 *
+		 * @return True if one or more strongly connected components are added
+		 *         to OutComponents. False otherwise. 
+		 */
+		static bool TarjanStronglyConnectedComponents(const FDirectedGraphAlgoAdapter& InAdapter, TArray<FStronglyConnectedComponent>& OutComponents, bool bExcludeSingleVertex = true);
+
+		/** Finds all nodes which can be reached by traversing the graph beginning
+		 * from FInputDataDestination nodes. 
+		 *
+		 * @param InGraph  - The graph to be traversed.
+		 * @param OutNodes - The nodes which were reached during traversal.
+		 */
+		static void FindReachableNodesFromInput(const IGraph& InGraph, TSet<const INode*>& OutNodes);
+
+		/** Finds all nodes which can be reached by traversing the graph beginning
+		 * from FInputDataDestination nodes. 
+		 *
+		 * @param InAdapter - The adapter to be traversed.
+		 * @param OutNodes  - The nodes which were reached during traversal.
+		 */
+		static void FindReachableNodesFromInput(const FDirectedGraphAlgoAdapter& InAdapter, TSet<const INode*>& OutNodes);
+
+		/** Finds all nodes which can be reached by traversing the transpose graph 
+		 * beginning from FOutputDataSource nodes. 
+		 *
+		 * @param InGraph  - The graph to be traversed.
+		 * @param OutNodes - The nodes which were reached during traversal.
+		 */
+		static void FindReachableNodesFromOutput(const IGraph& InGraph, TSet<const INode*>& OutNodes);
+
+		/** Finds all nodes which can be reached by traversing the transpose graph 
+		 * beginning from FOutputDataSource nodes. 
+		 *
+		 * @param InAdapter - The adapter to be traversed.
+		 * @param OutNodes  - The nodes which were reached during traversal.
+		 */
+		static void FindReachableNodesFromOutput(const FDirectedGraphAlgoAdapter& InAdapter, TSet<const INode*>& OutNodes);
+
+		/** Finds all nodes which can be reached by traversing the graph 
+		 * beginning from FInputDataDestinations or by traversing the transpose
+		 * graph beginning from FOutputDataSource nodes. 
+		 *
+		 * @param InGraph  - The graph to be traversed.
+		 * @param OutNodes - The nodes which were reached during traversal.
+		 */
+		static void FindReachableNodes(const IGraph& InGraph, TSet<const INode*>& OutNodes);
+
+		/** Finds all nodes which can be reached by traversing the graph 
+		 * beginning from FInputDataDestinations or by traversing the transpose
+		 * graph beginning from FOutputDataSource nodes. 
+		 *
+		 * @param InAdapter - The adapter to be traversed.
+		 * @param OutNodes  - The nodes which were reached during traversal.
+		 */
+		static void FindReachableNodes(const FDirectedGraphAlgoAdapter& InAdapter, TSet<const INode*>& OutNodes);
+
+	};
+}
diff --git a/Engine/Plugins/Runtime/Metasound/Source/MetasoundGraphCore/Public/MetasoundNodeInterface.h b/Engine/Plugins/Runtime/Metasound/Source/MetasoundGraphCore/Public/MetasoundNodeInterface.h
--- a/Engine/Plugins/Runtime/Metasound/Source/MetasoundGraphCore/Public/MetasoundNodeInterface.h
+++ b/Engine/Plugins/Runtime/Metasound/Source/MetasoundGraphCore/Public/MetasoundNodeInterface.h
@@ -237,6 +237,9 @@ namespace Metasound
 		,	Vertex(InVertex)
 		{
 		}
+
+		/** Check if two FOutputDataSources are equal. */
+		friend bool operator==(const FOutputDataSource& InLeft, const FOutputDataSource& InRight);
 	};
 
 	/** FInputDataSource describes the destination of data which produced 
@@ -259,6 +262,9 @@ namespace Metasound
 		,	Vertex(InVertex)
 		{
 		}
+
+		/** Check if two FInputDataDestinations are equal. */
+		friend bool operator==(const FInputDataDestination& InLeft, const FInputDataDestination& InRight);
 	};
 
 	/** Key type for an FOutputDataSource or FInputDataDestination. */
@@ -304,6 +310,9 @@ namespace Metasound
 		,	To(InTo)
 		{
 		}
+
+		/** Check if two FDataEdges are equal. */
+		friend bool operator==(const FDataEdge& InLeft, const FDataEdge& InRight);
 	};
 
 
diff --git a/Engine/Plugins/Runtime/Metasound/Source/MetasoundGraphCore/Public/MetasoundTarjan.h b/Engine/Plugins/Runtime/Metasound/Source/MetasoundGraphCore/Public/MetasoundTarjan.h
deleted file mode 100644
diff --git a/Engine/Source/Runtime/SignalProcessing/Private/DirectedGraphAlgo.cpp b/Engine/Source/Runtime/SignalProcessing/Private/DirectedGraphAlgo.cpp
new file mode 100644
--- /dev/null
+++ b/Engine/Source/Runtime/SignalProcessing/Private/DirectedGraphAlgo.cpp
@@ -0,0 +1,450 @@
+// Copyright Epic Games, Inc. All Rights Reserved.
+
+#include "DSP/DirectedGraphAlgo.h"
+
+namespace Audio
+{
+	namespace DirectedGraphAlgoPrivate
+	{
+		// Used to track which vertices have been visited and vertex children.
+		struct FDepthFirstSortInfo
+		{
+			bool TemporaryMark = false;
+			bool PermanentMark = false;
+			TArray<int32> Children;
+		};
+
+		// Recursive function for depth first topological sorting.
+		bool DepthFirstTopologicalSortVisit(int32 InVertex, TMap<int32, FDepthFirstSortInfo>& InSortInfos, TArray<int32>& OutReverseOrder)
+		{
+			FDepthFirstSortInfo& Info = InSortInfos[InVertex];
+
+			if (Info.PermanentMark)
+			{
+				// Exit depth first traversal
+				return true;
+			}
+			else if (Info.TemporaryMark)
+			{
+				// Found a cycle in the DAG.
+				return false;
+			}
+
+			Info.TemporaryMark = true;
+
+			for (int32 Child : Info.Children)
+			{
+				if (!DepthFirstTopologicalSortVisit(Child, InSortInfos, OutReverseOrder))
+				{
+					// cycle found in children.
+					return false;
+				}
+			}
+
+			Info.TemporaryMark = false;
+			Info.PermanentMark = true;
+
+			OutReverseOrder.Add(InVertex);
+
+			// subgraph successfully sorted.
+			return true;
+		}
+	} // End namespace DirectedAlgoGraphCorePrivate
+
+
+	void FDirectedGraphAlgo::BuildDirectedTree(TArrayView<const FDirectedEdge> InEdges, FDirectedTree& OutTree)
+	{
+		for (const FDirectedEdge& Edge : InEdges)
+		{
+			FDirectedTreeElement& Element = OutTree.FindOrAdd(Edge.Get<0>());
+
+			Element.Children.AddUnique(Edge.Get<1>());
+		}
+	}
+
+	void FDirectedGraphAlgo::BuildTransposeDirectedTree(TArrayView<const FDirectedEdge> InEdges, FDirectedTree& OutTree)
+	{
+		for (const FDirectedEdge& Edge : InEdges)
+		{
+			FDirectedTreeElement& Element = OutTree.FindOrAdd(Edge.Get<1>());
+
+			Element.Children.AddUnique(Edge.Get<0>());
+		}
+	}
+
+	void FDirectedGraphAlgo::DepthFirstTraversal(int32 InInitialVertex, const FDirectedTree& InTree, TFunctionRef<bool (int32)> InVisitFunc)
+	{
+		// Non recursive depth first traversal.
+
+		TArray<int32> Stack({InInitialVertex});
+		TSet<int32> Visited;
+
+		while (Stack.Num() > 0)
+		{
+			int32 CurrentVertex = Stack.Pop();
+			if (Visited.Contains(CurrentVertex))
+			{
+				continue;
+			}
+
+			bool bVisitChildren = InVisitFunc(CurrentVertex);
+
+			Visited.Add(CurrentVertex);
+
+			if (bVisitChildren)
+			{
+				const FDirectedTreeElement* Element = InTree.Find(CurrentVertex);
+
+				if (nullptr != Element)
+				{
+					Stack.Append(Element->Children);
+				}
+			}
+		}
+	}
+
+
+	void FDirectedGraphAlgo::BreadthFirstTraversal(int32 InInitialVertex, const FDirectedTree& InTree, TFunctionRef<bool (int32)> InVisitFunc)
+	{
+		// Non recursive breadth first traversal.
+
+		TArray<int32> Queue({InInitialVertex});
+		TSet<int32> Visited;
+
+		bool bVisitChildren = true;
+
+		while (Queue.Num() > 0)
+		{
+			int32 CurrentVertex = Queue[0];
+			Queue.RemoveAt(0);
+
+			if (Visited.Contains(CurrentVertex))
+			{
+				continue;
+			}
+
+			bVisitChildren = InVisitFunc(CurrentVertex);
+
+			Visited.Add(CurrentVertex);
+
+			if (bVisitChildren)
+			{
+				const FDirectedTreeElement* Element = InTree.Find(CurrentVertex);
+
+				if (nullptr != Element)
+				{
+					Queue.Append(Element->Children);
+				}
+			}
+		}
+	}
+
+	bool FDirectedGraphAlgo::DepthFirstTopologicalSort(TArrayView<const int32> InUniqueVertices, TArrayView<const FDirectedEdge> InUniqueEdges, TArray<int32>& OutVertexOrder)
+	{
+		using namespace DirectedGraphAlgoPrivate;
+
+		
+		// Initialize side info about which vertices have been visited and which
+		// have not.
+		TMap<int32, FDepthFirstSortInfo> SortInfos;
+		for (int32 Vertex : InUniqueVertices)
+		{
+			SortInfos.Add(Vertex);
+		}
+
+		for (const FDirectedEdge& Edge : InUniqueEdges)
+		{
+			FDepthFirstSortInfo& Info = SortInfos[Edge.Get<0>()];
+			Info.Children.Add(Edge.Get<1>());
+		}
+
+		bool bSuccess = true;
+
+		// Loop through each node to make sure they are all visited.
+		for (int32 Vertex : InUniqueVertices)
+		{
+			FDepthFirstSortInfo& OuterInfo = SortInfos[Vertex];
+
+			if (!(OuterInfo.TemporaryMark || OuterInfo.PermanentMark))
+			{
+				bSuccess = DepthFirstTopologicalSortVisit(Vertex, SortInfos, OutVertexOrder);
+
+				if (!bSuccess)
+				{
+					break;
+				}
+			}
+		}
+
+		// Depth first sort generates it in reverse order. It is reversed here
+		// so it is in forward order. 
+		Algo::Reverse(OutVertexOrder);
+
+		return bSuccess;
+	}
+
+	bool FDirectedGraphAlgo::KahnTopologicalSort(TArrayView<const int32> InUniqueVertices, TArrayView<const FDirectedEdge> InUniqueEdges, TArray<int32>& OutVertexOrder)
+	{
+		using FVertexMultiMap = TMultiMap<int32, int32>;
+		using FVertexPair = TMultiMap<int32, int32>::ElementType;
+		
+		FVertexMultiMap Dependencies;
+
+		for (const FDirectedEdge& Edge : InUniqueEdges)
+		{
+			// Destination vertices are dependent upon source vertices.
+			// Create a multimap where each destination vertex refrences
+			// all of the source vertices that have not been placed in the
+			// output. 
+			Dependencies.Add(Edge.Get<1>(), Edge.Get<0>());
+		}
+
+		TArray<int32> UniqueVertices(InUniqueVertices.GetData(), InUniqueVertices.Num());
+		TArray<int32> IndependentVertices;
+
+		// Sort graph so that vertices with no dependencies always go first.
+		while (UniqueVertices.Num() > 0)
+		{
+			IndependentVertices = UniqueVertices.FilterByPredicate([&](int32 InVertex) { 
+					return (Dependencies.Num(InVertex) == 0);
+			});
+
+			if (0 == IndependentVertices.Num())
+			{
+				// If there are no independent vertices, then likely there is a 
+				// cycle in the graph.
+				return false;
+			}
+
+			FVertexMultiMap UpdatedDependencies;
+			
+			// Remove independent vertices from dependency map.
+			for (FVertexPair& Element : Dependencies)
+			{
+				if (!IndependentVertices.Contains(Element.Value))
+				{
+					UpdatedDependencies.AddUnique(Element.Key, Element.Value);
+				}
+			}
+			Dependencies = UpdatedDependencies;
+
+			// Remove independent vertices from node list
+			for (int32 Vertex : IndependentVertices)
+			{
+				UniqueVertices.RemoveSwap(Vertex);
+			}
+
+			// Add independent vertices to output.
+			OutVertexOrder.Append(IndependentVertices);
+		}
+
+		return true;
+	}
+
+	namespace DirectedGraphAlgoPrivate
+	{
+		struct FTarjanVertexInfo
+		{
+			// Tag value given to a vertex.
+			int32 Tag = INDEX_NONE;
+
+			// Lowest tag in strongly connected components
+			int32 LowLinkTag = INDEX_NONE;
+
+			// Flag for global stack
+			bool bIsOnStack = false;
+
+			// Immediate children of vertex in DAG.
+			TArray<int32> Children;
+		};
+
+		struct FTarjanAlgoImplSettings 
+		{
+			// Single vertices can be returned as strongly connected components
+			// but usually they are not wanted by the caller. Setting this to
+			// true will require all strongly connected components to contain
+			// two or more vertices.
+			bool bExcludeSingleVertex;
+		};
+
+		// Private implementation of algorithm.
+		class FTarjanAlgoImpl
+		{
+			public:
+
+			FTarjanAlgoImpl(const FTarjanAlgoImplSettings& InSettings)
+			:	Settings(InSettings)
+			{
+			}
+
+			bool FindStronglyConnectedComponents(const TSet<FDirectedEdge>& InEdges, TArray<FStronglyConnectedComponent>& OutComponents)
+			{
+				using FElement = TMap<int32, FTarjanVertexInfo>::ElementType;
+
+				// Initialize internal data stores and tracking info.
+				Init(InEdges);
+				
+				int32 OriginalNumOutComponents = OutComponents.Num();
+
+				// Iterate through all vertices. Update(...) does a 
+				// depth-first-search which may result some vertices getting 
+				// inspected during recursive calls to Update(...).
+				for (FElement& Element : VertexInfos)
+				{
+					// Only examine vertice if it hasn't already been inspected.
+					if (Element.Value.Tag == INDEX_NONE)
+					{
+						Update(Element.Key, Element.Value, OutComponents);
+					}
+				}
+				
+				// Check if the size of OutComponents has grown.
+				return OutComponents.Num() > OriginalNumOutComponents;
+			}
+
+			private:
+
+			void Init(const TSet<FDirectedEdge>& InEdges)
+			{
+				CurrentTag = 0;
+
+				VertexInfos.Reset();
+				VertexStack.Reset();
+
+				for (const FDirectedEdge& Edge : InEdges)
+				{
+					FTarjanVertexInfo& Info = VertexInfos.FindOrAdd(Edge.Get<0>());
+					Info.Children.AddUnique(Edge.Get<1>());
+
+					VertexInfos.FindOrAdd(Edge.Get<1>());
+				}
+			}
+
+			void Update(int32 InVertex, FTarjanVertexInfo& InVertexInfo, TArray<FStronglyConnectedComponent>& OutComponents)
+			{
+				InVertexInfo.Tag = CurrentTag;
+				InVertexInfo.LowLinkTag = CurrentTag;
+				InVertexInfo.bIsOnStack = true;
+				VertexStack.Add(InVertex);
+
+				// Tag needs to always be increasing during subsequent calls to
+				// Update(...).
+				CurrentTag++;
+
+				// Perform depth first traversal through the graph looking for 
+				// children that are already on the stack (which denotes a 
+				// strongly connected component)
+				for (int32 ChildVertex : InVertexInfo.Children)
+				{
+					FTarjanVertexInfo& ChildInfo = VertexInfos[ChildVertex];
+
+					if (ChildInfo.Tag == INDEX_NONE)
+					{
+						// If ChildInfo.Tag == INDEX_NONE, then this child is 
+						// not on the stack because it has never been inspected.
+
+						// Continue in depth-first-traversal. 
+						Update(ChildVertex, ChildInfo, OutComponents);
+
+						// Update to use the lowest link tag found in children.
+						InVertexInfo.LowLinkTag = FMath::Min(InVertexInfo.LowLinkTag, ChildInfo.LowLinkTag);
+					}
+					else if(ChildInfo.bIsOnStack)
+					{
+						// This child is on the stack and represents a strongly 
+						// connected component. 
+						InVertexInfo.LowLinkTag = FMath::Min(InVertexInfo.LowLinkTag, ChildInfo.Tag);
+					}
+				}
+
+				if (InVertexInfo.LowLinkTag == InVertexInfo.Tag)
+				{
+					// Create a strongly connected component object if this is 
+					// the root vertex of a strongly connected component.
+					AddStronglyConnectedComponent(InVertex, OutComponents);
+				}
+			}
+
+			// Adds the current strongly connected component to the 
+			// OutComponents array.
+			void AddStronglyConnectedComponent(int32 InRootVertex, TArray<FStronglyConnectedComponent>& OutComponents)
+			{
+				// When a single vertex represents the entire strongly connected
+				// component, the root vertex is the same as the vertex at the
+				// top of the stack.
+				bool bIsSingleVertex = (VertexStack.Top() == InRootVertex);
+
+				if (bIsSingleVertex && Settings.bExcludeSingleVertex)
+				{
+					// This strongly connected component is skipped because it 
+					// consists of a single vertex. 
+					//
+					// The stack still needs to be updated as well as the vertex
+					// info for the root vertex.
+					int32 TopVertex = VertexStack.Pop();	
+					FTarjanVertexInfo& TopInfo = VertexInfos[TopVertex];
+					TopInfo.bIsOnStack = false;
+				}
+				else
+				{
+					bool bDidCompleteLoop = false;
+
+					// Create a new strongly connected component subgraph.
+					FStronglyConnectedComponent& StronglyConnectedComponent = OutComponents.AddDefaulted_GetRef();
+
+					// Pop off the top of the vertex stack until we reach the
+					// root vertex.
+					while (!bDidCompleteLoop)
+					{
+						int32 TopVertex = VertexStack.Pop();
+
+						FTarjanVertexInfo& TopInfo = VertexInfos[TopVertex];
+						TopInfo.bIsOnStack = false;
+
+						StronglyConnectedComponent.Vertices.Add(TopVertex);
+
+						// The unwinding continues until the current vertex is 
+						// removed from the stack.
+						bDidCompleteLoop = (TopVertex == InRootVertex);
+					}
+
+					// Add the edges associated with the strongly connected
+					// component. 
+					for (int32 ComponentVertex : StronglyConnectedComponent.Vertices)
+					{
+						FTarjanVertexInfo& ComponentInfo = VertexInfos[ComponentVertex];
+
+						for (int32 ChildVertex : ComponentInfo.Children)
+						{
+							if (StronglyConnectedComponent.Vertices.Contains(ChildVertex))
+							{
+								StronglyConnectedComponent.Edges.Emplace(ComponentVertex, ChildVertex);
+							}
+						}
+					}
+				}
+			}
+
+			FTarjanAlgoImplSettings Settings;
+
+			int32 CurrentTag = 0;
+
+			TMap<int32, FTarjanVertexInfo> VertexInfos;
+
+			TArray<int32> VertexStack;
+		};
+	}
+
+	bool FDirectedGraphAlgo::TarjanStronglyConnectedComponents(const TSet<FDirectedEdge>& InEdges, TArray<FStronglyConnectedComponent>& OutComponents, bool bExcludeSingleVertex)
+	{
+		using namespace DirectedGraphAlgoPrivate;
+
+		// Use private implementation of Tarjan algorithm.
+		FTarjanAlgoImplSettings Settings;
+		Settings.bExcludeSingleVertex = bExcludeSingleVertex;
+
+		FTarjanAlgoImpl TarjanAlgo(Settings);
+
+		return TarjanAlgo.FindStronglyConnectedComponents(InEdges, OutComponents);
+	}
+}
diff --git a/Engine/Source/Runtime/SignalProcessing/Public/DSP/DirectedGraphAlgo.h b/Engine/Source/Runtime/SignalProcessing/Public/DSP/DirectedGraphAlgo.h
new file mode 100644
--- /dev/null
+++ b/Engine/Source/Runtime/SignalProcessing/Public/DSP/DirectedGraphAlgo.h
@@ -0,0 +1,112 @@
+// Copyright Epic Games, Inc. All Rights Reserved.
+
+#pragma once
+
+#include "CoreMinimal.h"
+
+namespace Audio
+{
+	/** A pair of int32s represent a directed edge. The first value represents 
+	 * the source vertex, and the second value represents the destination 
+	 * vertex.
+	 */
+	typedef TTuple<int32, int32> FDirectedEdge;
+
+	/** A strongly connected component contains a subgraph of strongly connected
+	 * vertices and their corresponding edges.
+	 */
+	struct FStronglyConnectedComponent
+	{
+		/** Vertices in the strongly connected component. */
+		TArray<int32> Vertices;
+
+		/** Edges in the strongly connected component. */
+		TArray<FDirectedEdge> Edges;
+	};
+
+	/** An element in a directed tree with references to children of a vertex. */
+	struct FDirectedTreeElement
+	{
+		TArray<int32> Children;
+	};
+
+	/** A directed tree graph represenation. */
+	typedef TMap<int32, FDirectedTreeElement> FDirectedTree;
+
+
+	struct SIGNALPROCESSING_API FDirectedGraphAlgo
+	{
+		/** Build a directed tree from an array of edges.
+		 *
+		 * @parma InEdges - An array of directed eges.
+		 * @param OutTree - A tree structure built from the edges.
+		 */
+		static void BuildDirectedTree(TArrayView<const FDirectedEdge> InEdges, FDirectedTree& OutTree);
+
+		/** Build a transpose directed tree from an array of edges. 
+		 *
+		 * The transpose of a directed graph is created by reversing each edge.
+		 *
+		 * @parma InEdges - An array of directed eges.
+		 * @param OutTree - A tree structure built from the reversed edges.
+		 */
+		static void BuildTransposeDirectedTree(TArrayView<const FDirectedEdge> InEdges, FDirectedTree& OutTree);
+
+		/** Traverse a tree in a depth first ordering.
+		 *
+		 * @param InInitialVertex - The starting vertex for traversal.
+		 * @param InTree          - The tree structure to traverse. 
+		 * @param InVisitFunc     - A function that is called for each vertex that 
+		 *                          is visited. If this function returns true, 
+		 *                          traversal is continued in a depth first 
+		 *                          manner. If this function returns false, the 
+		 *                          children of the current vertex are not visited.
+		 */
+		static void DepthFirstTraversal(int32 InInitialVertex, const FDirectedTree& InTree, TFunctionRef<bool (int32)> InVisitFunc);
+
+		/** Traverse a tree in a breadth first ordering.
+		 *
+		 * @param InInitialVertex - The starting vertex for traversal.
+		 * @param InTree          - The tree structure to traverse. 
+		 * @param InVisitFunc     - A function that is called for each vertex that 
+		 *                          is visited. If this function returns true, 
+		 *                          the children of this vertex will be visited 
+		 *                          in the future.  If this function returns false, 
+		 *                          then the children of the current vertex will 
+		 *                          not be visited.
+		 */
+		static void BreadthFirstTraversal(int32 InInitialVertex, const FDirectedTree& InTree, TFunctionRef<bool (int32)> InVisitFunc);
+
+		/** Sort vertices topologically using a depth first sorting algorithm.
+		 *
+		 * @param InUniqueVertices - An array of vertices to sort.
+		 * @param InUniqueEdges - An array of edges describing dependencies.
+		 * @param OutVertexOrder - An array where ordered vertices are placed.
+		 *
+		 * @return True if sorting was successful. False otherwise.
+		 */
+		static bool DepthFirstTopologicalSort(TArrayView<const int32> InUniqueVertices, TArrayView<const FDirectedEdge> InUniqueEdges, TArray<int32>& OutVertexOrder);
+
+		/** Sort vertices topologically using a Kahn's sorting algorithm.
+		 *
+		 * @param InUniqueVertices - An array of vertices to sort.
+		 * @param InUniqueEdges - An array of edges describing dependencies.
+		 * @param OutVertexOrder - An array where ordered vertices are placed.
+		 *
+		 * @return True if sorting was successful. False otherwise.
+		 */
+		static bool KahnTopologicalSort(TArrayView<const int32> InUniqueVertices, TArrayView<const FDirectedEdge> InUniqueEdges, TArray<int32>& OutVertexOrder);
+
+
+		/** Find strongly connected components given a set of edges using Tarjan
+		 * algorithm.
+		 *
+		 * @param InEdges - Edges in the directed graph.
+		 * @param OutComponents - Strongly connected components found in the graph are be added to this array.
+		 * @param bExcludeSingleVertex - If true, single vertices are not be returned as strongly connected components. If false, single vertices may be returned as strongly connected components. 
+		 *
+		 * @return True if one or more strongly connected components are added to OutComponents. False otherwise. 
+		 */
+		static bool TarjanStronglyConnectedComponents(const TSet<FDirectedEdge>& InEdges, TArray<FStronglyConnectedComponent>& OutComponents, bool bExcludeSingleVertex = true);
+	};
+}
