Add FAligned versions in favor of code standard and deprecate old variants

diff --git a/Engine/Source/Runtime/SignalProcessing/Private/BufferVectorOperations.cpp b/Engine/Source/Runtime/SignalProcessing/Private/BufferVectorOperations.cpp
--- a/Engine/Source/Runtime/SignalProcessing/Private/BufferVectorOperations.cpp
+++ b/Engine/Source/Runtime/SignalProcessing/Private/BufferVectorOperations.cpp
@@ -13,7 +13,7 @@ namespace Audio
 	}
 
 	/* Sets a values to zero if value is denormal. Denormal numbers significantly slow down floating point operations. */
-	void BufferUnderflowClampFast(AlignedFloatBuffer& InOutBuffer)
+	void BufferUnderflowClampFast(FAlignedFloatBuffer& InOutBuffer)
 	{
 		BufferUnderflowClampFast(InOutBuffer.GetData(), InOutBuffer.Num());
 	}
@@ -73,7 +73,7 @@ namespace Audio
 		}
 	}
 
-	void BufferMultiplyByConstant(const AlignedFloatBuffer& InFloatBuffer, float InValue, AlignedFloatBuffer& OutFloatBuffer)
+	void BufferMultiplyByConstant(const FAlignedFloatBuffer& InFloatBuffer, float InValue, FAlignedFloatBuffer& OutFloatBuffer)
 	{
 		check(InFloatBuffer.Num() >= 4);
 
@@ -181,7 +181,7 @@ namespace Audio
 	}
 
 	/* Performs an element-wise weighted sum OutputBuffer = (InBuffer1 x InGain1) + (InBuffer2 x InGain2) */
-	void BufferWeightedSumFast(const AlignedFloatBuffer& InBuffer1, float InGain1, const AlignedFloatBuffer& InBuffer2, float InGain2, AlignedFloatBuffer& OutBuffer)
+	void BufferWeightedSumFast(const FAlignedFloatBuffer& InBuffer1, float InGain1, const FAlignedFloatBuffer& InBuffer2, float InGain2, FAlignedFloatBuffer& OutBuffer)
 	{
 		checkf(InBuffer1.Num() == InBuffer2.Num(), TEXT("Buffers must be equal length"));
 		OutBuffer.Reset();
@@ -191,7 +191,7 @@ namespace Audio
 	}
 
 	/* Performs an element-wise weighted sum OutputBuffer = (InBuffer1 x InGain1) + InBuffer2 */
-	void BufferWeightedSumFast(const AlignedFloatBuffer& InBuffer1, float InGain1, const AlignedFloatBuffer& InBuffer2, AlignedFloatBuffer& OutBuffer)
+	void BufferWeightedSumFast(const FAlignedFloatBuffer& InBuffer1, float InGain1, const FAlignedFloatBuffer& InBuffer2, FAlignedFloatBuffer& OutBuffer)
 	{
 		checkf(InBuffer1.Num() == InBuffer2.Num(), TEXT("Buffers must be equal length"));
 		OutBuffer.Reset();
@@ -335,7 +335,7 @@ namespace Audio
 #endif
 	}
 
-	void MixInBufferFast(const AlignedFloatBuffer& InFloatBuffer, AlignedFloatBuffer& BufferToSumTo, const float Gain)
+	void MixInBufferFast(const FAlignedFloatBuffer& InFloatBuffer, FAlignedFloatBuffer& BufferToSumTo, const float Gain)
 	{
 		MixInBufferFast(InFloatBuffer.GetData(), BufferToSumTo.GetData(), InFloatBuffer.Num(), Gain);
 	}
@@ -366,7 +366,7 @@ namespace Audio
 #endif
 	}
 
-	void MixInBufferFast(const AlignedFloatBuffer& InFloatBuffer, AlignedFloatBuffer& BufferToSumTo)
+	void MixInBufferFast(const FAlignedFloatBuffer& InFloatBuffer, FAlignedFloatBuffer& BufferToSumTo)
 	{
 		checkf(InFloatBuffer.Num() == BufferToSumTo.Num(), TEXT("Buffers must be equal size"));
 		MixInBufferFast(InFloatBuffer.GetData(), BufferToSumTo.GetData(), InFloatBuffer.Num());
@@ -396,7 +396,7 @@ namespace Audio
 #endif
 	}
 
-	void MixInBufferFast(const AlignedFloatBuffer& InFloatBuffer, AlignedFloatBuffer& BufferToSumTo, const float StartGain, const float EndGain)
+	void MixInBufferFast(const FAlignedFloatBuffer& InFloatBuffer, FAlignedFloatBuffer& BufferToSumTo, const float StartGain, const float EndGain)
 	{
 		MixInBufferFast(InFloatBuffer.GetData(), BufferToSumTo.GetData(), InFloatBuffer.Num(), StartGain, EndGain);
 	}
@@ -456,7 +456,7 @@ namespace Audio
 	}
 
 	/* Subtracts two buffers together element-wise. */
-	void BufferSubtractFast(const AlignedFloatBuffer& InMinuend, const AlignedFloatBuffer& InSubtrahend, AlignedFloatBuffer& OutputBuffer)
+	void BufferSubtractFast(const FAlignedFloatBuffer& InMinuend, const FAlignedFloatBuffer& InSubtrahend, FAlignedFloatBuffer& OutputBuffer)
 	{
 		const int32 InNum = InMinuend.Num();
 		OutputBuffer.Reset(InNum);
@@ -489,7 +489,7 @@ namespace Audio
 	}
 
 	/* Performs element-wise in-place subtraction placing the result in the subtrahend. InOutSubtrahend = InMinuend - InOutSubtrahend */
-	void BufferSubtractInPlace1Fast(const AlignedFloatBuffer& InMinuend, AlignedFloatBuffer& InOutSubtrahend)
+	void BufferSubtractInPlace1Fast(const FAlignedFloatBuffer& InMinuend, FAlignedFloatBuffer& InOutSubtrahend)
 	{
 		checkf(InMinuend.Num() == InOutSubtrahend.Num(), TEXT("Input buffers must be equal length"));
 		BufferSubtractInPlace1Fast(InMinuend.GetData(), InOutSubtrahend.GetData(), InMinuend.Num());
@@ -516,7 +516,7 @@ namespace Audio
 	}
 	
 	/* Performs element-wise in-place subtraction placing the result in the minuend. InOutMinuend = InOutMinuend - InSubtrahend */
-	void BufferSubtractInPlace2Fast(AlignedFloatBuffer& InOutMinuend, const AlignedFloatBuffer& InSubtrahend)
+	void BufferSubtractInPlace2Fast(AlignedFloatBuffer& InOutMinuend, const FAlignedFloatBuffer& InSubtrahend)
 	{
 		checkf(InOutMinuend.Num() == InSubtrahend.Num(), TEXT("Input buffers must be equal length"));
 		BufferSubtractInPlace2Fast(InOutMinuend.GetData(), InSubtrahend.GetData(), InOutMinuend.Num());
@@ -542,7 +542,7 @@ namespace Audio
 		}
 	}
 
-	void SumBuffers(const AlignedFloatBuffer& InFloatBuffer1, const AlignedFloatBuffer& InFloatBuffer2, AlignedFloatBuffer& OutputBuffer)
+	void SumBuffers(const FAlignedFloatBuffer& InFloatBuffer1, const FAlignedFloatBuffer& InFloatBuffer2, FAlignedFloatBuffer& OutputBuffer)
 	{
 		checkf(InFloatBuffer1.Num() == InFloatBuffer2.Num(), TEXT("Input buffers must be equal length"));
 		const int32 InNum = InFloatBuffer1.Num();
@@ -579,7 +579,7 @@ namespace Audio
 #endif
 	}
 
-	void MultiplyBuffersInPlace(const AlignedFloatBuffer& InFloatBuffer, AlignedFloatBuffer& BufferToMultiply)
+	void MultiplyBuffersInPlace(const FAlignedFloatBuffer& InFloatBuffer, FAlignedFloatBuffer& BufferToMultiply)
 	{
 		MultiplyBuffersInPlace(InFloatBuffer.GetData(), BufferToMultiply.GetData(), BufferToMultiply.Num());
 	}
@@ -602,7 +602,7 @@ namespace Audio
 		}
 	}
 
-	float GetMagnitude(const AlignedFloatBuffer& Buffer)
+	float GetMagnitude(const FAlignedFloatBuffer& Buffer)
 	{
 		return GetMagnitude(Buffer.GetData(), Buffer.Num());
 	}
@@ -637,7 +637,7 @@ namespace Audio
 #endif
 	}
 
-	float GetAverageAmplitude(const AlignedFloatBuffer& Buffer)
+	float GetAverageAmplitude(const FAlignedFloatBuffer& Buffer)
 	{
 		checkf(Buffer.Num() % 4 == 0, TEXT("Please use a buffer size that is a multiple of 4."));
 
@@ -770,7 +770,7 @@ namespace Audio
 		}
 	}
 
-	void MixMonoTo2ChannelsFast(const AlignedFloatBuffer& MonoBuffer, AlignedFloatBuffer& DestinationBuffer, const float* RESTRICT Gains)
+	void MixMonoTo2ChannelsFast(const FAlignedFloatBuffer& MonoBuffer, FAlignedFloatBuffer& DestinationBuffer, const float* RESTRICT Gains)
 	{
 		MixMonoTo2ChannelsFast(MonoBuffer.GetData(), DestinationBuffer.GetData(), DestinationBuffer.Num() / 2, Gains);
 	}
@@ -799,7 +799,7 @@ namespace Audio
 		}
 	}
 
-	void MixMonoTo2ChannelsFast(const AlignedFloatBuffer& MonoBuffer, AlignedFloatBuffer& DestinationBuffer, const float* RESTRICT StartGains, const float* RESTRICT EndGains)
+	void MixMonoTo2ChannelsFast(const FAlignedFloatBuffer& MonoBuffer, FAlignedFloatBuffer& DestinationBuffer, const float* RESTRICT StartGains, const float* RESTRICT EndGains)
 	{
 		MixMonoTo2ChannelsFast(MonoBuffer.GetData(), DestinationBuffer.GetData(), DestinationBuffer.Num() / 2, StartGains, EndGains);
 	}
@@ -840,7 +840,7 @@ namespace Audio
 		}
 	}
 
-	void MixMonoTo2ChannelsFast(const AlignedFloatBuffer& MonoBuffer, AlignedFloatBuffer& DestinationBuffer)
+	void MixMonoTo2ChannelsFast(const FAlignedFloatBuffer& MonoBuffer, FAlignedFloatBuffer& DestinationBuffer)
 	{
 		MixMonoTo2ChannelsFast(MonoBuffer.GetData(), DestinationBuffer.GetData(), MonoBuffer.Num());
 	}
@@ -876,7 +876,7 @@ namespace Audio
 		}
 	}
 
-	void Mix2ChannelsTo2ChannelsFast(const AlignedFloatBuffer& SourceBuffer, AlignedFloatBuffer& DestinationBuffer, const float* RESTRICT Gains)
+	void Mix2ChannelsTo2ChannelsFast(const FAlignedFloatBuffer& SourceBuffer, FAlignedFloatBuffer& DestinationBuffer, const float* RESTRICT Gains)
 	{
 		Mix2ChannelsTo2ChannelsFast(SourceBuffer.GetData(), DestinationBuffer.GetData(), DestinationBuffer.Num() / 2, Gains);
 	}
@@ -916,7 +916,7 @@ namespace Audio
 		}
 	}
 
-	void Mix2ChannelsTo2ChannelsFast(const AlignedFloatBuffer& SourceBuffer, AlignedFloatBuffer& DestinationBuffer, const float* RESTRICT StartGains, const float* RESTRICT EndGains)
+	void Mix2ChannelsTo2ChannelsFast(const FAlignedFloatBuffer& SourceBuffer, FAlignedFloatBuffer& DestinationBuffer, const float* RESTRICT StartGains, const float* RESTRICT EndGains)
 	{
 		Mix2ChannelsTo2ChannelsFast(SourceBuffer.GetData(), DestinationBuffer.GetData(), DestinationBuffer.Num() / 2, StartGains, EndGains);
 	}
@@ -1015,7 +1015,7 @@ namespace Audio
 		}
 	}
 
-	void MixMonoTo4ChannelsFast(const AlignedFloatBuffer& MonoBuffer, AlignedFloatBuffer& DestinationBuffer, const float* RESTRICT Gains)
+	void MixMonoTo4ChannelsFast(const FAlignedFloatBuffer& MonoBuffer, FAlignedFloatBuffer& DestinationBuffer, const float* RESTRICT Gains)
 	{
 		MixMonoTo4ChannelsFast(MonoBuffer.GetData(), DestinationBuffer.GetData(), DestinationBuffer.Num() / 4, Gains);
 	}
@@ -1046,7 +1046,7 @@ namespace Audio
 		}
 	}
 
-	void MixMonoTo4ChannelsFast(const AlignedFloatBuffer& MonoBuffer, AlignedFloatBuffer& DestinationBuffer, const float* RESTRICT StartGains, const float* RESTRICT EndGains)
+	void MixMonoTo4ChannelsFast(const FAlignedFloatBuffer& MonoBuffer, FAlignedFloatBuffer& DestinationBuffer, const float* RESTRICT StartGains, const float* RESTRICT EndGains)
 	{
 		MixMonoTo4ChannelsFast(MonoBuffer.GetData(), DestinationBuffer.GetData(), DestinationBuffer.Num() / 4, StartGains, EndGains);
 	}
@@ -1082,7 +1082,7 @@ namespace Audio
 		}
 	}
 
-	void Mix2ChannelsTo4ChannelsFast(const AlignedFloatBuffer& SourceBuffer, AlignedFloatBuffer& DestinationBuffer, const float* RESTRICT Gains)
+	void Mix2ChannelsTo4ChannelsFast(const FAlignedFloatBuffer& SourceBuffer, FAlignedFloatBuffer& DestinationBuffer, const float* RESTRICT Gains)
 	{
 		Mix2ChannelsTo4ChannelsFast(SourceBuffer.GetData(), DestinationBuffer.GetData(), DestinationBuffer.Num() / 4, Gains);
 	}
@@ -1121,7 +1121,7 @@ namespace Audio
 		}
 	}
 
-	void Mix2ChannelsTo4ChannelsFast(const AlignedFloatBuffer& SourceBuffer, AlignedFloatBuffer& DestinationBuffer, const float* RESTRICT StartGains, const float* RESTRICT EndGains)
+	void Mix2ChannelsTo4ChannelsFast(const FAlignedFloatBuffer& SourceBuffer, FAlignedFloatBuffer& DestinationBuffer, const float* RESTRICT StartGains, const float* RESTRICT EndGains)
 	{
 		Mix2ChannelsTo4ChannelsFast(SourceBuffer.GetData(), DestinationBuffer.GetData(), DestinationBuffer.Num() / 4, StartGains, EndGains);
 	}
@@ -1239,7 +1239,7 @@ namespace Audio
 		}
 	}
 
-	void MixMonoTo6ChannelsFast(const AlignedFloatBuffer& MonoBuffer, AlignedFloatBuffer& DestinationBuffer, const float* RESTRICT Gains)
+	void MixMonoTo6ChannelsFast(const FAlignedFloatBuffer& MonoBuffer, FAlignedFloatBuffer& DestinationBuffer, const float* RESTRICT Gains)
 	{
 		MixMonoTo6ChannelsFast(MonoBuffer.GetData(), DestinationBuffer.GetData(), DestinationBuffer.Num() / 6, Gains);
 	}
@@ -1282,7 +1282,7 @@ namespace Audio
 		}
 	}
 
-	void MixMonoTo6ChannelsFast(const AlignedFloatBuffer& MonoBuffer, AlignedFloatBuffer& DestinationBuffer, const float* RESTRICT StartGains, const float* RESTRICT EndGains)
+	void MixMonoTo6ChannelsFast(const FAlignedFloatBuffer& MonoBuffer, FAlignedFloatBuffer& DestinationBuffer, const float* RESTRICT StartGains, const float* RESTRICT EndGains)
 	{
 		MixMonoTo6ChannelsFast(MonoBuffer.GetData(), DestinationBuffer.GetData(), DestinationBuffer.Num() / 6, StartGains, EndGains);
 	}
@@ -1341,7 +1341,7 @@ namespace Audio
 		}
 	}
 
-	void Mix2ChannelsTo6ChannelsFast(const AlignedFloatBuffer& SourceBuffer, AlignedFloatBuffer& DestinationBuffer, const float* RESTRICT Gains)
+	void Mix2ChannelsTo6ChannelsFast(const FAlignedFloatBuffer& SourceBuffer, FAlignedFloatBuffer& DestinationBuffer, const float* RESTRICT Gains)
 	{
 		Mix2ChannelsTo6ChannelsFast(SourceBuffer.GetData(), DestinationBuffer.GetData(), DestinationBuffer.Num() / 6, Gains);
 	}
@@ -1402,7 +1402,7 @@ namespace Audio
 		}
 	}
 
-	void Mix2ChannelsTo6ChannelsFast(const AlignedFloatBuffer& SourceBuffer, AlignedFloatBuffer& DestinationBuffer, const float* RESTRICT StartGains, const float* RESTRICT EndGains)
+	void Mix2ChannelsTo6ChannelsFast(const FAlignedFloatBuffer& SourceBuffer, FAlignedFloatBuffer& DestinationBuffer, const float* RESTRICT StartGains, const float* RESTRICT EndGains)
 	{
 		Mix2ChannelsTo6ChannelsFast(SourceBuffer.GetData(), DestinationBuffer.GetData(), DestinationBuffer.Num() / 6, StartGains, EndGains);
 	}
@@ -1563,7 +1563,7 @@ namespace Audio
 		}
 	}
 
-	void MixMonoTo8ChannelsFast(const AlignedFloatBuffer& MonoBuffer, AlignedFloatBuffer& DestinationBuffer, const float* RESTRICT Gains)
+	void MixMonoTo8ChannelsFast(const FAlignedFloatBuffer& MonoBuffer, FAlignedFloatBuffer& DestinationBuffer, const float* RESTRICT Gains)
 	{
 		MixMonoTo8ChannelsFast(MonoBuffer.GetData(), DestinationBuffer.GetData(), DestinationBuffer.Num() / 8, Gains);
 	}
@@ -1599,7 +1599,7 @@ namespace Audio
 		}
 	}
 
-	void MixMonoTo8ChannelsFast(const AlignedFloatBuffer& MonoBuffer, AlignedFloatBuffer& DestinationBuffer, const float* RESTRICT StartGains, const float* RESTRICT EndGains)
+	void MixMonoTo8ChannelsFast(const FAlignedFloatBuffer& MonoBuffer, FAlignedFloatBuffer& DestinationBuffer, const float* RESTRICT StartGains, const float* RESTRICT EndGains)
 	{
 		MixMonoTo8ChannelsFast(MonoBuffer.GetData(), DestinationBuffer.GetData(), DestinationBuffer.Num() / 8, StartGains, EndGains);
 	}
@@ -1644,7 +1644,7 @@ namespace Audio
 		}
 	}
 
-	void Mix2ChannelsTo8ChannelsFast(const AlignedFloatBuffer& SourceBuffer, AlignedFloatBuffer& DestinationBuffer, const float* RESTRICT Gains)
+	void Mix2ChannelsTo8ChannelsFast(const FAlignedFloatBuffer& SourceBuffer, FAlignedFloatBuffer& DestinationBuffer, const float* RESTRICT Gains)
 	{
 		Mix2ChannelsTo8ChannelsFast(SourceBuffer.GetData(), DestinationBuffer.GetData(), DestinationBuffer.Num() / 8, Gains);
 	}
@@ -1691,7 +1691,7 @@ namespace Audio
 		}
 	}
 
-	void Mix2ChannelsTo8ChannelsFast(const AlignedFloatBuffer& SourceBuffer, AlignedFloatBuffer& DestinationBuffer, const float* RESTRICT StartGains, const float* RESTRICT EndGains)
+	void Mix2ChannelsTo8ChannelsFast(const FAlignedFloatBuffer& SourceBuffer, FAlignedFloatBuffer& DestinationBuffer, const float* RESTRICT StartGains, const float* RESTRICT EndGains)
 	{
 		Mix2ChannelsTo8ChannelsFast(SourceBuffer.GetData(), DestinationBuffer.GetData(), DestinationBuffer.Num() / 8, StartGains, EndGains);
 	}
@@ -1759,7 +1759,7 @@ namespace Audio
 	/**
 	 * These functions are non-vectorized versions of the Mix[N]ChannelsTo[N]Channels functions above:
 	 */
-	void DownmixBuffer(int32 NumSourceChannels, int32 NumDestinationChannels, const AlignedFloatBuffer& SourceBuffer, AlignedFloatBuffer& DestinationBuffer, const float* RESTRICT Gains)
+	void DownmixBuffer(int32 NumSourceChannels, int32 NumDestinationChannels, const FAlignedFloatBuffer& SourceBuffer, FAlignedFloatBuffer& DestinationBuffer, const float* RESTRICT Gains)
 	{
 		DownmixBuffer(NumSourceChannels, NumDestinationChannels, SourceBuffer.GetData(), DestinationBuffer.GetData(), SourceBuffer.Num() / NumSourceChannels, Gains);
 	}
@@ -1782,7 +1782,7 @@ namespace Audio
 		}
 	}
 
-	void DownmixBuffer(int32 NumSourceChannels, int32 NumDestinationChannels, const AlignedFloatBuffer& SourceBuffer, AlignedFloatBuffer& DestinationBuffer, float* RESTRICT StartGains, const float* RESTRICT EndGains)
+	void DownmixBuffer(int32 NumSourceChannels, int32 NumDestinationChannels, const FAlignedFloatBuffer& SourceBuffer, FAlignedFloatBuffer& DestinationBuffer, float* RESTRICT StartGains, const float* RESTRICT EndGains)
 	{
 		DownmixBuffer(NumSourceChannels, NumDestinationChannels, SourceBuffer.GetData(), DestinationBuffer.GetData(), SourceBuffer.Num() / NumSourceChannels, StartGains, EndGains);
 	}
@@ -1821,7 +1821,7 @@ namespace Audio
 	}
 
 
-	SIGNALPROCESSING_API void DownmixAndSumIntoBuffer(int32 NumSourceChannels, int32 NumDestinationChannels, const AlignedFloatBuffer& SourceBuffer, AlignedFloatBuffer& BufferToSumTo, const float* RESTRICT Gains)
+	SIGNALPROCESSING_API void DownmixAndSumIntoBuffer(int32 NumSourceChannels, int32 NumDestinationChannels, const FAlignedFloatBuffer& SourceBuffer, FAlignedFloatBuffer& BufferToSumTo, const float* RESTRICT Gains)
 	{
 		DownmixAndSumIntoBuffer(NumSourceChannels, NumDestinationChannels, SourceBuffer.GetData(), BufferToSumTo.GetData(), SourceBuffer.Num() / NumSourceChannels, Gains);
 	}
@@ -1844,7 +1844,7 @@ namespace Audio
 	}
 
 	/** Interleaves samples from two input buffers */
-	void BufferInterleave2ChannelFast(const AlignedFloatBuffer& InBuffer1, const AlignedFloatBuffer& InBuffer2, AlignedFloatBuffer& OutBuffer)
+	void BufferInterleave2ChannelFast(const FAlignedFloatBuffer& InBuffer1, const FAlignedFloatBuffer& InBuffer2, FAlignedFloatBuffer& OutBuffer)
 	{
 		checkf(InBuffer1.Num() == InBuffer2.Num(), TEXT("InBuffer1 Num not equal to InBuffer2 Num"));
 
@@ -1892,7 +1892,7 @@ namespace Audio
 	}
 
 	/** Deinterleaves samples from a 2 channel input buffer */
-	void BufferDeinterleave2ChannelFast(const AlignedFloatBuffer& InBuffer, AlignedFloatBuffer& OutBuffer1, AlignedFloatBuffer& OutBuffer2)
+	void BufferDeinterleave2ChannelFast(const FAlignedFloatBuffer& InBuffer, FAlignedFloatBuffer& OutBuffer1, FAlignedFloatBuffer& OutBuffer2)
 	{
 		const int32 InNum = InBuffer.Num();
 		const int32 InNumFrames = InNum / 2;
@@ -1936,7 +1936,7 @@ namespace Audio
 	}
 
 	/** Sums 2 channel interleaved input samples. OutSamples[n] = InSamples[2n] + InSamples[2n + 1] */
-	void BufferSum2ChannelToMonoFast(const AlignedFloatBuffer& InSamples, AlignedFloatBuffer& OutSamples)
+	void BufferSum2ChannelToMonoFast(const FAlignedFloatBuffer& InSamples, FAlignedFloatBuffer& OutSamples)
 	{
 		const int32 InNum = InSamples.Num();
 		const int32 Frames = InNum / 2;
@@ -1980,7 +1980,7 @@ namespace Audio
 		}
 	}
 
-	void BufferComplexToPowerFast(const AlignedFloatBuffer& InRealSamples, const AlignedFloatBuffer& InImaginarySamples, AlignedFloatBuffer& OutPowerSamples)
+	void BufferComplexToPowerFast(const FAlignedFloatBuffer& InRealSamples, const FAlignedFloatBuffer& InImaginarySamples, FAlignedFloatBuffer& OutPowerSamples)
 	{
 		checkf(InRealSamples.Num() == InImaginarySamples.Num(), TEXT("Input buffers must have equal number of elements"));
 
@@ -2020,14 +2020,14 @@ namespace Audio
 
 	// class FBufferLinearEase implementation
 	FBufferLinearEase::FBufferLinearEase() {}
-	FBufferLinearEase::FBufferLinearEase(const AlignedFloatBuffer & InSourceValues, const AlignedFloatBuffer & InTargetValues, int32 InLerpLength)
+	FBufferLinearEase::FBufferLinearEase(const FAlignedFloatBuffer & InSourceValues, const FAlignedFloatBuffer & InTargetValues, int32 InLerpLength)
 	{
 		Init(InSourceValues, InTargetValues, InLerpLength);
 	}
 
 	FBufferLinearEase::~FBufferLinearEase() {}
 
-	void FBufferLinearEase::Init(const AlignedFloatBuffer & InSourceValues, const AlignedFloatBuffer & InTargetValues, int32 InLerpLength)
+	void FBufferLinearEase::Init(const FAlignedFloatBuffer & InSourceValues, const FAlignedFloatBuffer & InTargetValues, int32 InLerpLength)
 	{
 		check(InSourceValues.Num());
 		check(InTargetValues.Num());
@@ -2064,7 +2064,7 @@ namespace Audio
 		return false;
 	}
 
-	bool FBufferLinearEase::Update(uint32 StepsToJumpForward, AlignedFloatBuffer & InSourceValues)
+	bool FBufferLinearEase::Update(uint32 StepsToJumpForward, FAlignedFloatBuffer & InSourceValues)
 	{
 		check(InSourceValues.Num() == BufferLength);
 		check(CurrentLerpStep != LerpLength);
@@ -2083,7 +2083,7 @@ namespace Audio
 		return bIsComplete;
 
 	}
-	const AlignedFloatBuffer & FBufferLinearEase::GetDeltaBuffer()
+	const FAlignedFloatBuffer & FBufferLinearEase::GetDeltaBuffer()
 	{
 		return DeltaBuffer;
 	}
diff --git a/Engine/Source/Runtime/SignalProcessing/Public/DSP/BufferVectorOperations.h b/Engine/Source/Runtime/SignalProcessing/Public/DSP/BufferVectorOperations.h
--- a/Engine/Source/Runtime/SignalProcessing/Public/DSP/BufferVectorOperations.h
+++ b/Engine/Source/Runtime/SignalProcessing/Public/DSP/BufferVectorOperations.h
@@ -17,109 +17,114 @@
 namespace Audio
 {
 	/** Aligned allocator used for fast operations. */	
-	typedef TAlignedHeapAllocator<AUDIO_BUFFER_ALIGNMENT> FAudioBufferAlignedAllocator;
+	using FAudioBufferAlignedAllocator = TAlignedHeapAllocator<AUDIO_BUFFER_ALIGNMENT>;
 
-	typedef TArray<float, FAudioBufferAlignedAllocator> AlignedFloatBuffer;
+	using FAlignedByteBuffer  = TArray<uint8, FAudioBufferAlignedAllocator>;
+	using FAlignedFloatBuffer = TArray<float, FAudioBufferAlignedAllocator>;
+	using FAlignedInt32Buffer = TArray<int32, FAudioBufferAlignedAllocator>;
+
+	// Deprecated in favor of versions above
 	typedef TArray<uint8, FAudioBufferAlignedAllocator> AlignedByteBuffer;
+	typedef TArray<float, FAudioBufferAlignedAllocator> AlignedFloatBuffer;
 	typedef TArray<int32, FAudioBufferAlignedAllocator> AlignedInt32Buffer;
 
 
 	/** CHANNEL-AGNOSTIC OPERATIONS */
 
 	/* Sets a values to zero if value is denormal. Denormal numbers significantly slow down floating point operations. */
-	SIGNALPROCESSING_API void BufferUnderflowClampFast(AlignedFloatBuffer& InOutBuffer);
+	SIGNALPROCESSING_API void BufferUnderflowClampFast(FAlignedFloatBuffer& InOutBuffer);
 
 	/* Sets a values to zero if value is denormal. Denormal numbers significantly slow down floating point operations. */
 	SIGNALPROCESSING_API void BufferUnderflowClampFast(float* RESTRICT InOutBuffer, const int32 InNum);
 	
 	/* Clamps the values in a buffer between a min and max value.*/
-	SIGNALPROCESSING_API void BufferRangeClampFast(AlignedFloatBuffer& InOutBuffer, float InMinValue, float InMaxValue);
+	SIGNALPROCESSING_API void BufferRangeClampFast(FAlignedFloatBuffer& InOutBuffer, float InMinValue, float InMaxValue);
 	SIGNALPROCESSING_API void BufferRangeClampFast(float* RESTRICT InOutBuffer, const int32 InNum, float InMinValue, float InMaxValue);
 	
 	/** Multiplies the input aligned float buffer with the given value. */
-	SIGNALPROCESSING_API void BufferMultiplyByConstant(const AlignedFloatBuffer& InFloatBuffer, float InValue, AlignedFloatBuffer& OutFloatBuffer);
+	SIGNALPROCESSING_API void BufferMultiplyByConstant(const FAlignedFloatBuffer& InFloatBuffer, float InValue, FAlignedFloatBuffer& OutFloatBuffer);
 	SIGNALPROCESSING_API void BufferMultiplyByConstant(const float* RESTRICT InFloatBuffer, float InValue, float* RESTRICT OutFloatBuffer, const int32 InNumSamples);
-	SIGNALPROCESSING_API void BufferMultiplyByConstant(const AlignedFloatBuffer& InFloatBuffer, float InValue, AlignedFloatBuffer& OutFloatBuffer);
+	SIGNALPROCESSING_API void BufferMultiplyByConstant(const FAlignedFloatBuffer& InFloatBuffer, float InValue, FAlignedFloatBuffer& OutFloatBuffer);
 
 	/** Similar to BufferMultiplyByConstant, but (a) assumes a buffer length divisible by 4 and (b) performs the multiply in place. */
-	SIGNALPROCESSING_API void MultiplyBufferByConstantInPlace(AlignedFloatBuffer& InBuffer, float InGain);
+	SIGNALPROCESSING_API void MultiplyBufferByConstantInPlace(FAlignedFloatBuffer& InBuffer, float InGain);
 	SIGNALPROCESSING_API void MultiplyBufferByConstantInPlace(float* RESTRICT InBuffer, int32 NumSamples, float InGain);
 
 	/** Adds a constant to a buffer (useful for DC offset removal) */
-	SIGNALPROCESSING_API void AddConstantToBufferInplace(AlignedFloatBuffer& InBuffer, float Constant);
+	SIGNALPROCESSING_API void AddConstantToBufferInplace(FAlignedFloatBuffer& InBuffer, float Constant);
 	SIGNALPROCESSING_API void AddConstantToBufferInplace(float* RESTRICT InBuffer, int32 NumSamples, float Constant);
 
 	/** Sets a constant to a buffer (useful for DC offset application) */
-	SIGNALPROCESSING_API void BufferSetToConstantInplace(AlignedFloatBuffer& InBuffer, float Constant);
+	SIGNALPROCESSING_API void BufferSetToConstantInplace(FAlignedFloatBuffer& InBuffer, float Constant);
 	SIGNALPROCESSING_API void BufferSetToConstantInplace(float* RESTRICT InBuffer, int32 NumSamples, float Constant);
 
 	/* Performs an element-wise weighted sum OutputBuffer = (InBuffer1 x InGain1) + (InBuffer2 x InGain2) */
-	SIGNALPROCESSING_API void BufferWeightedSumFast(const AlignedFloatBuffer& InBuffer1, float InGain1, const AlignedFloatBuffer& InBuffer2, float InGain2, AlignedFloatBuffer& OutBuffer);
+	SIGNALPROCESSING_API void BufferWeightedSumFast(const FAlignedFloatBuffer& InBuffer1, float InGain1, const FAlignedFloatBuffer& InBuffer2, float InGain2, FAlignedFloatBuffer& OutBuffer);
 
 	/* Performs an element-wise weighted sum OutputBuffer = (InBuffer1 x InGain1) + InBuffer2 */
-	SIGNALPROCESSING_API void BufferWeightedSumFast(const AlignedFloatBuffer& InBuffer1, float InGain1, const AlignedFloatBuffer& InBuffer2, AlignedFloatBuffer& OutBuffer);
+	SIGNALPROCESSING_API void BufferWeightedSumFast(const FAlignedFloatBuffer& InBuffer1, float InGain1, const FAlignedFloatBuffer& InBuffer2, FAlignedFloatBuffer& OutBuffer);
 
 	/* Performs an element-wise weighted sum OutputBuffer = (InBuffer1 x InGain1) + (InBuffer2 x InGain2) */
 	SIGNALPROCESSING_API void BufferWeightedSumFast(const float* RESTRICT InBuffer1, float InGain1, const float* RESTRICT InBuffer2, float InGain2, float* RESTRICT OutBuffer, int32 InNum);
 
 	/* Performs an element-wise weighted sum OutputBuffer = (InBuffer1 x InGain1) + InBuffer2 */
 	SIGNALPROCESSING_API void BufferWeightedSumFast(const float* RESTRICT InBuffer1, float InGain1, const float* RESTRICT InBuffer2, float* RESTRICT OutBuffer, int32 InNum);
-	SIGNALPROCESSING_API void MultiplyBufferByConstantInPlace(AlignedFloatBuffer& InBuffer, float InGain);
+	SIGNALPROCESSING_API void MultiplyBufferByConstantInPlace(FAlignedFloatBuffer& InBuffer, float InGain);
 	SIGNALPROCESSING_API void MultiplyBufferByConstantInPlace(float* RESTRICT InBuffer, int32 NumSamples, float InGain);
 
 	/* Takes a float buffer and quickly interpolates it's gain from StartValue to EndValue. */
 	/* This operation completely ignores channel counts, so avoid using this function on buffers that are not mono, stereo or quad */
 	/* if the buffer needs to fade all channels uniformly. */
-	SIGNALPROCESSING_API void FadeBufferFast(AlignedFloatBuffer& OutFloatBuffer, const float StartValue, const float EndValue);
+	SIGNALPROCESSING_API void FadeBufferFast(FAlignedFloatBuffer& OutFloatBuffer, const float StartValue, const float EndValue);
 	SIGNALPROCESSING_API void FadeBufferFast(float* RESTRICT OutFloatBuffer, int32 NumSamples, const float StartValue, const float EndValue);
 
 	/** Takes buffer InFloatBuffer, optionally multiplies it by Gain, and adds it to BufferToSumTo. */
-	SIGNALPROCESSING_API void MixInBufferFast(const AlignedFloatBuffer& InFloatBuffer, AlignedFloatBuffer& BufferToSumTo, const float Gain);
+	SIGNALPROCESSING_API void MixInBufferFast(const FAlignedFloatBuffer& InFloatBuffer, FAlignedFloatBuffer& BufferToSumTo, const float Gain);
 	SIGNALPROCESSING_API void MixInBufferFast(const float* RESTRICT InFloatBuffer, float* RESTRICT BufferToSumTo, int32 NumSamples, const float Gain);
-	SIGNALPROCESSING_API void MixInBufferFast(const AlignedFloatBuffer& InFloatBuffer, AlignedFloatBuffer& BufferToSumTo, const float StartGain, const float EndGain);
+	SIGNALPROCESSING_API void MixInBufferFast(const FAlignedFloatBuffer& InFloatBuffer, FAlignedFloatBuffer& BufferToSumTo, const float StartGain, const float EndGain);
 	SIGNALPROCESSING_API void MixInBufferFast(const float* RESTRICT InFloatBuffer, float* RESTRICT BufferToSumTo, int32 NumSamples, const float StartGain, const float EndGain);
-	SIGNALPROCESSING_API void MixInBufferFast(const AlignedFloatBuffer& InFloatBuffer, AlignedFloatBuffer& BufferToSumTo);
+	SIGNALPROCESSING_API void MixInBufferFast(const FAlignedFloatBuffer& InFloatBuffer, FAlignedFloatBuffer& BufferToSumTo);
 	SIGNALPROCESSING_API void MixInBufferFast(const float* RESTRICT InFloatBuffer, float* RESTRICT BufferToSumTo, int32 NumSamples);
 
 	/* Subtracts two buffers together element-wise. */
-	SIGNALPROCESSING_API void BufferSubtractFast(const AlignedFloatBuffer& InMinuend, const AlignedFloatBuffer& InSubtrahend, AlignedFloatBuffer& OutputBuffer);
+	SIGNALPROCESSING_API void BufferSubtractFast(const FAlignedFloatBuffer& InMinuend, const FAlignedFloatBuffer& InSubtrahend, FAlignedFloatBuffer& OutputBuffer);
 
 	/* Subtracts two buffers together element-wise. */
 	SIGNALPROCESSING_API void BufferSubtractFast(const float* RESTRICT InMinuend, const float* RESTRICT InSubtrahend, float* RESTRICT OutputBuffer, int32 NumSamples);
 
 	/* Performs element-wise in-place subtraction placing the result in the subtrahend. InOutSubtrahend = InMinuend - InOutSubtrahend */
-	SIGNALPROCESSING_API void BufferSubtractInPlace1Fast(const AlignedFloatBuffer& InMinuend, AlignedFloatBuffer& InOutSubtrahend);
+	SIGNALPROCESSING_API void BufferSubtractInPlace1Fast(const FAlignedFloatBuffer& InMinuend, FAlignedFloatBuffer& InOutSubtrahend);
 
 	/* Performs element-wise in-place subtraction placing the result in the subtrahend. InOutSubtrahend = InMinuend - InOutSubtrahend */
 	SIGNALPROCESSING_API void BufferSubtractInPlace1Fast(const float* RESTRICT InMinuend, float* RESTRICT InOutSubtrahend, int32 NumSamples);
 
 
 	/* Performs element-wise in-place subtraction placing the result in the minuend. InOutMinuend = InOutMinuend - InSubtrahend */
-	SIGNALPROCESSING_API void BufferSubtractInPlace2Fast(AlignedFloatBuffer& InOutMinuend, const AlignedFloatBuffer& InSubtrahend);
+	SIGNALPROCESSING_API void BufferSubtractInPlace2Fast(FAlignedFloatBuffer& InOutMinuend, const FAlignedFloatBuffer& InSubtrahend);
 
 	/* Performs element-wise in-place subtraction placing the result in the minuend. InOutMinuend = InOutMinuend - InSubtrahend */
 	SIGNALPROCESSING_API void BufferSubtractInPlace2Fast(float* RESTRICT InOutMinuend, const float* RESTRICT InSubtrahend, int32 NumSamples);
 
 	/** This version of MixInBufferFast will fade from StartGain to EndGain. */
-	SIGNALPROCESSING_API void MixInBufferFast(const AlignedFloatBuffer& InFloatBuffer, AlignedFloatBuffer& BufferToSumTo, const float StartGain, const float EndGain);
+	SIGNALPROCESSING_API void MixInBufferFast(const FAlignedFloatBuffer& InFloatBuffer, FAlignedFloatBuffer& BufferToSumTo, const float StartGain, const float EndGain);
 	SIGNALPROCESSING_API void MixInBufferFast(const float* RESTRICT InFloatBuffer, float* RESTRICT BufferToSumTo, int32 NumSamples, const float StartGain, const float EndGain);
 
 	/** Sums two buffers together and places the result in the resulting buffer. */
-	SIGNALPROCESSING_API void SumBuffers(const AlignedFloatBuffer& InFloatBuffer1, const AlignedFloatBuffer& InFloatBuffer2, AlignedFloatBuffer& OutputBuffer);
+	SIGNALPROCESSING_API void SumBuffers(const FAlignedFloatBuffer& InFloatBuffer1, const FAlignedFloatBuffer& InFloatBuffer2, FAlignedFloatBuffer& OutputBuffer);
 	SIGNALPROCESSING_API void SumBuffers(const float* RESTRICT InFloatBuffer1, const float* RESTRICT InFloatBuffer2, float* RESTRICT OutputBuffer, int32 NumSamples);
 
 	/** Multiply the second buffer in place by the first buffer. */
-	SIGNALPROCESSING_API void MultiplyBuffersInPlace(const AlignedFloatBuffer& InFloatBuffer, AlignedFloatBuffer& BufferToMultiply);
+	SIGNALPROCESSING_API void MultiplyBuffersInPlace(const FAlignedFloatBuffer& InFloatBuffer, FAlignedFloatBuffer& BufferToMultiply);
 	SIGNALPROCESSING_API void MultiplyBuffersInPlace(const float* RESTRICT InFloatBuffer, float* RESTRICT BufferToMultiply, int32 NumSamples);
 
 	/** CHANNEL-AGNOSTIC ANALYSIS OPERATIONS */
 
 	/** Takes an audio buffer and returns the magnitude across that buffer. */
-	SIGNALPROCESSING_API float GetMagnitude(const AlignedFloatBuffer& Buffer);
+	SIGNALPROCESSING_API float GetMagnitude(const FAlignedFloatBuffer& Buffer);
 	SIGNALPROCESSING_API float GetMagnitude(const float* RESTRICT Buffer, int32 NumSamples);
 
 	/** Takes an audio buffer and gets the average absolute amplitude across that buffer. */
-	SIGNALPROCESSING_API float GetAverageAmplitude(const AlignedFloatBuffer& Buffer);
+	SIGNALPROCESSING_API float GetAverageAmplitude(const FAlignedFloatBuffer& Buffer);
 	SIGNALPROCESSING_API float GetAverageAmplitude(const float* RESTRICT Buffer, int32 NumSamples);
 
 	/** CHANNEL-SPECIFIC OPERATIONS */
@@ -128,20 +133,20 @@ namespace Audio
 	 *  StereoBuffer must have an even number of frames.
 	 *  If StartGains and EndGains are provided, this function will interpolate between the two across the buffer.
 	*/
-	SIGNALPROCESSING_API void Apply2ChannelGain(AlignedFloatBuffer& StereoBuffer, const float* RESTRICT Gains);
+	SIGNALPROCESSING_API void Apply2ChannelGain(FAlignedFloatBuffer& StereoBuffer, const float* RESTRICT Gains);
 	SIGNALPROCESSING_API void Apply2ChannelGain(float* RESTRICT StereoBuffer, int32 NumSamples, const float* RESTRICT Gains);
-	SIGNALPROCESSING_API void Apply2ChannelGain(AlignedFloatBuffer& StereoBuffer, const float* RESTRICT StartGains, const float* RESTRICT EndGains);
+	SIGNALPROCESSING_API void Apply2ChannelGain(FAlignedFloatBuffer& StereoBuffer, const float* RESTRICT StartGains, const float* RESTRICT EndGains);
 	SIGNALPROCESSING_API void Apply2ChannelGain(float* RESTRICT StereoBuffer, int32 NumSamples, const float* RESTRICT StartGains, const float* RESTRICT EndGains);
 
 	/** Takes a 1 channel buffer and mixes it to a stereo buffer using Gains. Gains is expected to point to a 2 float long buffer.
 	 *  these buffers must have an even number of frames.
 	 *  If StartGains and EndGains are provided, this function will interpolate between the two across the buffer.
 	*/
-	SIGNALPROCESSING_API void MixMonoTo2ChannelsFast(const AlignedFloatBuffer& MonoBuffer, AlignedFloatBuffer& DestinationBuffer, const float* RESTRICT Gains);
+	SIGNALPROCESSING_API void MixMonoTo2ChannelsFast(const FAlignedFloatBuffer& MonoBuffer, FAlignedFloatBuffer& DestinationBuffer, const float* RESTRICT Gains);
 	SIGNALPROCESSING_API void MixMonoTo2ChannelsFast(const float* RESTRICT MonoBuffer, float* RESTRICT DestinationBuffer, int32 NumFrames, const float* RESTRICT Gains);
-	SIGNALPROCESSING_API void MixMonoTo2ChannelsFast(const AlignedFloatBuffer& MonoBuffer, AlignedFloatBuffer& DestinationBuffer, const float* RESTRICT StartGains, const float* RESTRICT EndGains);
+	SIGNALPROCESSING_API void MixMonoTo2ChannelsFast(const FAlignedFloatBuffer& MonoBuffer, FAlignedFloatBuffer& DestinationBuffer, const float* RESTRICT StartGains, const float* RESTRICT EndGains);
 	SIGNALPROCESSING_API void MixMonoTo2ChannelsFast(const float* RESTRICT MonoBuffer, float* RESTRICT DestinationBuffer, int32 NumFrames, const float* RESTRICT StartGains, const float* RESTRICT EndGains);
-	SIGNALPROCESSING_API void MixMonoTo2ChannelsFast(const AlignedFloatBuffer& MonoBuffer, AlignedFloatBuffer& DestinationBuffer);
+	SIGNALPROCESSING_API void MixMonoTo2ChannelsFast(const FAlignedFloatBuffer& MonoBuffer, FAlignedFloatBuffer& DestinationBuffer);
 	SIGNALPROCESSING_API void MixMonoTo2ChannelsFast(const float* RESTRICT MonoBuffer, float* RESTRICT DestinationBuffer, int32 NumFrames);
 
 	/** Takes a 2 channel buffer and mixes it to an 2 channel interleaved buffer using Gains. Gains is expected to point to a 16 float long buffer.
@@ -149,26 +154,26 @@ namespace Audio
 	*  NumFrames must be a multiple of 4.
 	*  If StartGains and EndGains are provided, this function will interpolate between the two across the buffer.
 	*/
-	SIGNALPROCESSING_API void Mix2ChannelsTo2ChannelsFast(const AlignedFloatBuffer& SourceBuffer, AlignedFloatBuffer& DestinationBuffer, const float* RESTRICT Gains);
+	SIGNALPROCESSING_API void Mix2ChannelsTo2ChannelsFast(const FAlignedFloatBuffer& SourceBuffer, FAlignedFloatBuffer& DestinationBuffer, const float* RESTRICT Gains);
 	SIGNALPROCESSING_API void Mix2ChannelsTo2ChannelsFast(const float* RESTRICT SourceBuffer, float* RESTRICT DestinationBuffer, int32 NumFrames, const float* RESTRICT Gains);
-	SIGNALPROCESSING_API void Mix2ChannelsTo2ChannelsFast(const AlignedFloatBuffer& SourceBuffer, AlignedFloatBuffer& DestinationBuffer, const float* RESTRICT StartGains, const float* RESTRICT EndGains);
+	SIGNALPROCESSING_API void Mix2ChannelsTo2ChannelsFast(const FAlignedFloatBuffer& SourceBuffer, FAlignedFloatBuffer& DestinationBuffer, const float* RESTRICT StartGains, const float* RESTRICT EndGains);
 	SIGNALPROCESSING_API void Mix2ChannelsTo2ChannelsFast(const float* RESTRICT SourceBuffer, float* RESTRICT DestinationBuffer, int32 NumFrames, const float* RESTRICT StartGains, const float* RESTRICT EndGains);
 
 	/** Takes a 4 channel interleaved buffer and applies Gains to it. Gains is expected to point to a 2 float long buffer.
 	 *  If StartGains and EndGains are provided, this function will interpolate between the two across the buffer.
 	*/
-	SIGNALPROCESSING_API void Apply4ChannelGain(AlignedFloatBuffer& InterleavedBuffer, const float* RESTRICT Gains);
+	SIGNALPROCESSING_API void Apply4ChannelGain(FAlignedFloatBuffer& InterleavedBuffer, const float* RESTRICT Gains);
 	SIGNALPROCESSING_API void Apply4ChannelGain(float* RESTRICT InterleavedBuffer, int32 NumSamples, const float* RESTRICT Gains);
-	SIGNALPROCESSING_API void Apply4ChannelGain(AlignedFloatBuffer& InterleavedBuffer, const float* RESTRICT StartGains, const float* RESTRICT EndGains);
+	SIGNALPROCESSING_API void Apply4ChannelGain(FAlignedFloatBuffer& InterleavedBuffer, const float* RESTRICT StartGains, const float* RESTRICT EndGains);
 	SIGNALPROCESSING_API void Apply4ChannelGain(float* RESTRICT InterleavedBuffer, int32 NumSamples, const float* RESTRICT StartGains, const float* RESTRICT EndGains);
 
 	/** Takes a 1 channel buffer and mixes it to an 8 channel interleaved buffer using Gains. Gains is expected to point to a 8 float long buffer.
 	*  these buffers must have an even number of frames.
 	*  If StartGains and EndGains are provided, this function will interpolate between the two across the buffer.
 	*/
-	SIGNALPROCESSING_API void MixMonoTo4ChannelsFast(const AlignedFloatBuffer& MonoBuffer, AlignedFloatBuffer& DestinationBuffer, const float* RESTRICT Gains);
+	SIGNALPROCESSING_API void MixMonoTo4ChannelsFast(const FAlignedFloatBuffer& MonoBuffer, FAlignedFloatBuffer& DestinationBuffer, const float* RESTRICT Gains);
 	SIGNALPROCESSING_API void MixMonoTo4ChannelsFast(const float* RESTRICT MonoBuffer, float* RESTRICT DestinationBuffer, int32 NumFrames, const float* RESTRICT Gains);
-	SIGNALPROCESSING_API void MixMonoTo4ChannelsFast(const AlignedFloatBuffer& MonoBuffer, AlignedFloatBuffer& DestinationBuffer, const float* RESTRICT StartGains, const float* RESTRICT EndGains);
+	SIGNALPROCESSING_API void MixMonoTo4ChannelsFast(const FAlignedFloatBuffer& MonoBuffer, FAlignedFloatBuffer& DestinationBuffer, const float* RESTRICT StartGains, const float* RESTRICT EndGains);
 	SIGNALPROCESSING_API void MixMonoTo4ChannelsFast(const float* RESTRICT MonoBuffer, float* RESTRICT DestinationBuffer, int32 NumFrames, const float* RESTRICT StartGains, const float* RESTRICT EndGains);
 
 	/** Takes a 2 channel buffer and mixes it to an 8 channel interleaved buffer using Gains. Gains is expected to point to a 16 float long buffer.
@@ -176,26 +181,26 @@ namespace Audio
 	*  NumFrames must be a multiple of 4.
 	*  If StartGains and EndGains are provided, this function will interpolate between the two across the buffer.
 	*/
-	SIGNALPROCESSING_API void Mix2ChannelsTo4ChannelsFast(const AlignedFloatBuffer& SourceBuffer, AlignedFloatBuffer& DestinationBuffer, const float* RESTRICT Gains);
+	SIGNALPROCESSING_API void Mix2ChannelsTo4ChannelsFast(const FAlignedFloatBuffer& SourceBuffer, FAlignedFloatBuffer& DestinationBuffer, const float* RESTRICT Gains);
 	SIGNALPROCESSING_API void Mix2ChannelsTo4ChannelsFast(const float* RESTRICT SourceBuffer, float* RESTRICT DestinationBuffer, int32 NumFrames, const float* RESTRICT Gains);
-	SIGNALPROCESSING_API void Mix2ChannelsTo4ChannelsFast(const AlignedFloatBuffer& SourceBuffer, AlignedFloatBuffer& DestinationBuffer, const float* RESTRICT StartGains, const float* RESTRICT EndGains);
+	SIGNALPROCESSING_API void Mix2ChannelsTo4ChannelsFast(const FAlignedFloatBuffer& SourceBuffer, FAlignedFloatBuffer& DestinationBuffer, const float* RESTRICT StartGains, const float* RESTRICT EndGains);
 	SIGNALPROCESSING_API void Mix2ChannelsTo4ChannelsFast(const float* RESTRICT SourceBuffer, float* RESTRICT DestinationBuffer, int32 NumFrames, const float* RESTRICT StartGains, const float* RESTRICT EndGains);
 
 	/** Takes a 6 channel interleaved buffer and applies Gains to it. Gains is expected to point to a 2 float long buffer.
 	 *  InterleavedBuffer must have an even number of frames.
 	*/
-	SIGNALPROCESSING_API void Apply6ChannelGain(AlignedFloatBuffer& InterleavedBuffer, const float* RESTRICT Gains);
+	SIGNALPROCESSING_API void Apply6ChannelGain(FAlignedFloatBuffer& InterleavedBuffer, const float* RESTRICT Gains);
 	SIGNALPROCESSING_API void Apply6ChannelGain(float* RESTRICT InterleavedBuffer, int32 NumSamples, const float* RESTRICT Gains);
-	SIGNALPROCESSING_API void Apply6ChannelGain(AlignedFloatBuffer& InterleavedBuffer, const float* RESTRICT StartGains, const float* RESTRICT EndGains);
+	SIGNALPROCESSING_API void Apply6ChannelGain(FAlignedFloatBuffer& InterleavedBuffer, const float* RESTRICT StartGains, const float* RESTRICT EndGains);
 	SIGNALPROCESSING_API void Apply6ChannelGain(float* RESTRICT InterleavedBuffer, int32 NumSamples, const float* RESTRICT StartGains, const float* RESTRICT EndGains);
 
 	/** Takes a 1 channel buffer and mixes it to an 8 channel interleaved buffer using Gains. Gains is expected to point to a 8 float long buffer.
 	 *  these buffers must have an even number of frames.
 	 *  If StartGains and EndGains are provided, this function will interpolate between the two across the buffer.
 	*/
-	SIGNALPROCESSING_API void MixMonoTo6ChannelsFast(const AlignedFloatBuffer& MonoBuffer, AlignedFloatBuffer& DestinationBuffer, const float* RESTRICT Gains);
+	SIGNALPROCESSING_API void MixMonoTo6ChannelsFast(const FAlignedFloatBuffer& MonoBuffer, FAlignedFloatBuffer& DestinationBuffer, const float* RESTRICT Gains);
 	SIGNALPROCESSING_API void MixMonoTo6ChannelsFast(const float* RESTRICT MonoBuffer, float* RESTRICT DestinationBuffer, int32 NumFrames, const float* RESTRICT Gains);
-	SIGNALPROCESSING_API void MixMonoTo6ChannelsFast(const AlignedFloatBuffer& MonoBuffer, AlignedFloatBuffer& DestinationBuffer, const float* RESTRICT StartGains, const float* RESTRICT EndGains);
+	SIGNALPROCESSING_API void MixMonoTo6ChannelsFast(const FAlignedFloatBuffer& MonoBuffer, FAlignedFloatBuffer& DestinationBuffer, const float* RESTRICT StartGains, const float* RESTRICT EndGains);
 	SIGNALPROCESSING_API void MixMonoTo6ChannelsFast(const float* RESTRICT MonoBuffer, float* RESTRICT DestinationBuffer, int32 NumFrames, const float* RESTRICT StartGains, const float* RESTRICT EndGains);
 
 	/** Takes a 2 channel buffer and mixes it to an 8 channel interleaved buffer using Gains. Gains is expected to point to a 16 float long buffer.
@@ -203,24 +208,24 @@ namespace Audio
 	 *  NumFrames must be a multiple of 4.
 	 *  If StartGains and EndGains are provided, this function will interpolate between the two across the buffer.
 	*/
-	SIGNALPROCESSING_API void Mix2ChannelsTo6ChannelsFast(const AlignedFloatBuffer& SourceBuffer, AlignedFloatBuffer& DestinationBuffer, const float* RESTRICT Gains);
+	SIGNALPROCESSING_API void Mix2ChannelsTo6ChannelsFast(const FAlignedFloatBuffer& SourceBuffer, FAlignedFloatBuffer& DestinationBuffer, const float* RESTRICT Gains);
 	SIGNALPROCESSING_API void Mix2ChannelsTo6ChannelsFast(const float* RESTRICT SourceBuffer, float* RESTRICT DestinationBuffer, int32 NumFrames, const float* RESTRICT Gains);
-	SIGNALPROCESSING_API void Mix2ChannelsTo6ChannelsFast(const AlignedFloatBuffer& SourceBuffer, AlignedFloatBuffer& DestinationBuffer, const float* RESTRICT StartGains, const float* RESTRICT EndGains);
+	SIGNALPROCESSING_API void Mix2ChannelsTo6ChannelsFast(const FAlignedFloatBuffer& SourceBuffer, FAlignedFloatBuffer& DestinationBuffer, const float* RESTRICT StartGains, const float* RESTRICT EndGains);
 	SIGNALPROCESSING_API void Mix2ChannelsTo6ChannelsFast(const float* RESTRICT SourceBuffer, float* RESTRICT DestinationBuffer, int32 NumFrames, const float* RESTRICT StartGains, const float* RESTRICT EndGains);
 
 	/** Takes an 8 channel interleaved buffer and applies Gains to it. Gains is expected to point to an 8 float long buffer. */
-	SIGNALPROCESSING_API void Apply8ChannelGain(AlignedFloatBuffer& InterleavedBuffer, const float* RESTRICT Gains);
+	SIGNALPROCESSING_API void Apply8ChannelGain(FAlignedFloatBuffer& InterleavedBuffer, const float* RESTRICT Gains);
 	SIGNALPROCESSING_API void Apply8ChannelGain(float* RESTRICT InterleavedBuffer, int32 NumSamples, const float* RESTRICT Gains);
-	SIGNALPROCESSING_API void Apply8ChannelGain(AlignedFloatBuffer& InterleavedBuffer, const float* RESTRICT StartGains, const float* RESTRICT EndGains);
+	SIGNALPROCESSING_API void Apply8ChannelGain(FAlignedFloatBuffer& InterleavedBuffer, const float* RESTRICT StartGains, const float* RESTRICT EndGains);
 	SIGNALPROCESSING_API void Apply8ChannelGain(float* RESTRICT InterleavedBuffer, int32 NumSamples, const float* RESTRICT StartGains, const float* RESTRICT EndGains);
 
 	/** Takes a 1 channel buffer and mixes it to an 8 channel interleaved buffer using Gains. Gains is expected to point to a 8 float long buffer.
 	 *  these buffers must have an even number of frames.
 	 *  If StartGains and EndGains are provided, this function will interpolate between the two across the buffer.
 	*/
-	SIGNALPROCESSING_API void MixMonoTo8ChannelsFast(const AlignedFloatBuffer& MonoBuffer, AlignedFloatBuffer& DestinationBuffer, const float* RESTRICT Gains);
+	SIGNALPROCESSING_API void MixMonoTo8ChannelsFast(const FAlignedFloatBuffer& MonoBuffer, FAlignedFloatBuffer& DestinationBuffer, const float* RESTRICT Gains);
 	SIGNALPROCESSING_API void MixMonoTo8ChannelsFast(const float* RESTRICT MonoBuffer, float* RESTRICT DestinationBuffer, int32 NumFrames, const float* RESTRICT Gains);
-	SIGNALPROCESSING_API void MixMonoTo8ChannelsFast(const AlignedFloatBuffer& MonoBuffer, AlignedFloatBuffer& DestinationBuffer, const float* RESTRICT StartGains, const float* RESTRICT EndGains);
+	SIGNALPROCESSING_API void MixMonoTo8ChannelsFast(const FAlignedFloatBuffer& MonoBuffer, FAlignedFloatBuffer& DestinationBuffer, const float* RESTRICT StartGains, const float* RESTRICT EndGains);
 	SIGNALPROCESSING_API void MixMonoTo8ChannelsFast(const float* RESTRICT MonoBuffer, float* RESTRICT DestinationBuffer, int32 NumFrames, const float* RESTRICT StartGains, const float* RESTRICT EndGains);
 
 	/** Takes a 2 channel buffer and mixes it to an 8 channel interleaved buffer using Gains. Gains is expected to point to a 16 float long buffer.
@@ -228,51 +233,51 @@ namespace Audio
 	 *  these buffers must have an even number of frames.
 	 *  If StartGains and EndGains are provided, this function will interpolate between the two across the buffer.
 	*/
-	SIGNALPROCESSING_API void Mix2ChannelsTo8ChannelsFast(const AlignedFloatBuffer& SourceBuffer, AlignedFloatBuffer& DestinationBuffer, const float* RESTRICT Gains);
+	SIGNALPROCESSING_API void Mix2ChannelsTo8ChannelsFast(const FAlignedFloatBuffer& SourceBuffer, FAlignedFloatBuffer& DestinationBuffer, const float* RESTRICT Gains);
 	SIGNALPROCESSING_API void Mix2ChannelsTo8ChannelsFast(const float* RESTRICT SourceBuffer, float* RESTRICT DestinationBuffer, int32 NumFrames, const float* RESTRICT Gains);
-	SIGNALPROCESSING_API void Mix2ChannelsTo8ChannelsFast(const AlignedFloatBuffer& SourceBuffer, AlignedFloatBuffer& DestinationBuffer, const float* RESTRICT StartGains, const float* RESTRICT EndGains);
+	SIGNALPROCESSING_API void Mix2ChannelsTo8ChannelsFast(const FAlignedFloatBuffer& SourceBuffer, FAlignedFloatBuffer& DestinationBuffer, const float* RESTRICT StartGains, const float* RESTRICT EndGains);
 	SIGNALPROCESSING_API void Mix2ChannelsTo8ChannelsFast(const float* RESTRICT SourceBuffer, float* RESTRICT DestinationBuffer, int32 NumFrames, const float* RESTRICT StartGains, const float* RESTRICT EndGains);
 
 	/** This is a generalized operation that uses the channel gain matrix provided in Gains to mix an interleaved source buffer to the interleaved downmix buffer.
 	 *  This operation is not explicitly vectorized and will almost always be slower than using one of the functions above.
 	*/
-	SIGNALPROCESSING_API void DownmixBuffer(int32 NumSourceChannels, int32 NumDestinationChannels, const AlignedFloatBuffer& SourceBuffer, AlignedFloatBuffer& DestinationBuffer, const float* RESTRICT Gains);
+	SIGNALPROCESSING_API void DownmixBuffer(int32 NumSourceChannels, int32 NumDestinationChannels, const FAlignedFloatBuffer& SourceBuffer, FAlignedFloatBuffer& DestinationBuffer, const float* RESTRICT Gains);
 	SIGNALPROCESSING_API void DownmixBuffer(int32 NumSourceChannels, int32 NumDestinationChannels, const float* RESTRICT SourceBuffer, float* RESTRICT DestinationBuffer, int32 NumFrames, const float* RESTRICT Gains);
-	SIGNALPROCESSING_API void DownmixBuffer(int32 NumSourceChannels, int32 NumDestinationChannels, const AlignedFloatBuffer& SourceBuffer, AlignedFloatBuffer& DestinationBuffer, float* RESTRICT StartGains, const float* RESTRICT EndGains);
+	SIGNALPROCESSING_API void DownmixBuffer(int32 NumSourceChannels, int32 NumDestinationChannels, const FAlignedFloatBuffer& SourceBuffer, FAlignedFloatBuffer& DestinationBuffer, float* RESTRICT StartGains, const float* RESTRICT EndGains);
 	SIGNALPROCESSING_API void DownmixBuffer(int32 NumSourceChannels, int32 NumDestinationChannels, const float* RESTRICT SourceBuffer, float* RESTRICT DestinationBuffer, int32 NumFrames, float* RESTRICT StartGains, const float* RESTRICT EndGains);
 
 	/**
 	 * This is similar to DownmixBuffer, except that it sums into DestinationBuffer rather than overwriting it.
 	 */
-	SIGNALPROCESSING_API void DownmixAndSumIntoBuffer(int32 NumSourceChannels, int32 NumDestinationChannels, const AlignedFloatBuffer& SourceBuffer, AlignedFloatBuffer& BufferToSumTo, const float* RESTRICT Gains);
+	SIGNALPROCESSING_API void DownmixAndSumIntoBuffer(int32 NumSourceChannels, int32 NumDestinationChannels, const FAlignedFloatBuffer& SourceBuffer, FAlignedFloatBuffer& BufferToSumTo, const float* RESTRICT Gains);
 	SIGNALPROCESSING_API void DownmixAndSumIntoBuffer(int32 NumSourceChannels, int32 NumDestinationChannels, const float* RESTRICT SourceBuffer, float* RESTRICT BufferToSumTo, int32 NumFrames, const float* RESTRICT Gains);
 
 	/** Interleaves samples from two input buffers */
-	SIGNALPROCESSING_API void BufferInterleave2ChannelFast(const AlignedFloatBuffer& InBuffer1, const AlignedFloatBuffer& InBuffer2, AlignedFloatBuffer& OutBuffer);
+	SIGNALPROCESSING_API void BufferInterleave2ChannelFast(const FAlignedFloatBuffer& InBuffer1, const FAlignedFloatBuffer& InBuffer2, FAlignedFloatBuffer& OutBuffer);
 
 	/** Interleaves samples from two input buffers */
 	SIGNALPROCESSING_API void BufferInterleave2ChannelFast(const float* RESTRICT InBuffer1, const float* RESTRICT InBuffer2, float* RESTRICT OutBuffer, const int32 InNum);
 
 	/** Deinterleaves samples from a 2 channel input buffer */
-	SIGNALPROCESSING_API void BufferDeinterleave2ChannelFast(const AlignedFloatBuffer& InBuffer, AlignedFloatBuffer& OutBuffer1, AlignedFloatBuffer& OutBuffer2);
+	SIGNALPROCESSING_API void BufferDeinterleave2ChannelFast(const FAlignedFloatBuffer& InBuffer, FAlignedFloatBuffer& OutBuffer1, FAlignedFloatBuffer& OutBuffer2);
 
 	/** Deinterleaves samples from a 2 channel input buffer */
 	SIGNALPROCESSING_API void BufferDeinterleave2ChannelFast(const float* RESTRICT InBuffer, float* RESTRICT OutBuffer1, float* RESTRICT OutBuffer2, const int32 InNum);
 
 	/** Sums 2 channel interleaved input samples. OutSamples[n] = InSamples[2n] + InSamples[2n + 1] */
-	SIGNALPROCESSING_API void BufferSum2ChannelToMonoFast(const AlignedFloatBuffer& InSamples, AlignedFloatBuffer& OutSamples);
+	SIGNALPROCESSING_API void BufferSum2ChannelToMonoFast(const FAlignedFloatBuffer& InSamples, FAlignedFloatBuffer& OutSamples);
 
 	/** Sums 2 channel interleaved input samples. OutSamples[n] = InSamples[2n] + InSamples[2n + 1] */
 	SIGNALPROCESSING_API void BufferSum2ChannelToMonoFast(const float* RESTRICT InSamples, float* RESTRICT OutSamples, const int32 InNumFrames);
 
 	/** Compute power of complex data. Out[i] = Real[i] * Real[i] + Imaginary[i] * Imaginary[i] */
-	SIGNALPROCESSING_API void BufferComplexToPowerFast(const AlignedFloatBuffer& InRealSamples, const AlignedFloatBuffer& InImaginarySamples, AlignedFloatBuffer& OutPowerSamples);
+	SIGNALPROCESSING_API void BufferComplexToPowerFast(const FAlignedFloatBuffer& InRealSamples, const FAlignedFloatBuffer& InImaginarySamples, FAlignedFloatBuffer& OutPowerSamples);
 
 	/** Compute power of complex data. Out[i] = Real[i] * Real[i] + Imaginary[i] * Imaginary[i] */
 	SIGNALPROCESSING_API void BufferComplexToPowerFast(const float* RESTRICT InRealSamples, const float* RESTRICT InImaginarySamples, float* RESTRICT OutPowerSamples, const int32 InNum);
 
 	/** Compute magnitude of complex data. Out[i] = Sqrt(Real[i] * Real[i] + Imaginary[i] * Imaginary[i]) */
-	SIGNALPROCESSING_API void BufferComplexToMagnitudeFast(const AlignedFloatBuffer& InRealSamples, const AlignedFloatBuffer& InImaginarySamples, AlignedFloatBuffer& OutPowerSamples);
+	SIGNALPROCESSING_API void BufferComplexToMagnitudeFast(const FAlignedFloatBuffer& InRealSamples, const FAlignedFloatBuffer& InImaginarySamples, FAlignedFloatBuffer& OutPowerSamples);
 
 	/** Compute magnitude of complex data. Out[i] = Sqrt(Real[i] * Real[i] + Imaginary[i] * Imaginary[i]) */
 	SIGNALPROCESSING_API void BufferComplexToMagnitudeFast(const float* RESTRICT InRealSamples, const float* RESTRICT InImaginarySamples, float* RESTRICT OutPowerSamples, const int32 InNum);
@@ -282,26 +287,26 @@ namespace Audio
 	{
 	public:
 		FBufferLinearEase();
-		FBufferLinearEase(const AlignedFloatBuffer& InSourceValues, const AlignedFloatBuffer& InTargetValues, int32 InLerpLength);
+		FBufferLinearEase(const FAlignedFloatBuffer& InSourceValues, const FAlignedFloatBuffer& InTargetValues, int32 InLerpLength);
 		~FBufferLinearEase();
 
 		/** will cache SourceValues ptr and manually update SourceValues on Update() */
-		void Init(const AlignedFloatBuffer& InSourceValues, const AlignedFloatBuffer& InTargetValues, int32 InLerpLength);
+		void Init(const FAlignedFloatBuffer& InSourceValues, const FAlignedFloatBuffer& InTargetValues, int32 InLerpLength);
 
 		/** Performs Vectorized update of SourceValues float buffer. Returns true if interpolation is complete */
-		bool Update(AlignedFloatBuffer& InSourceValues);
+		bool Update(FAlignedFloatBuffer& InSourceValues);
 
 		/** Update overloaded to let you jump forward more than a single time-step */
-		bool Update(uint32 StepsToJumpForward, AlignedFloatBuffer& InSourceValues);
+		bool Update(uint32 StepsToJumpForward, FAlignedFloatBuffer& InSourceValues);
 
 		/** returns const reference to the deltas buffer for doing interpolation elsewhere */
-		const AlignedFloatBuffer& GetDeltaBuffer();
+		const FAlignedFloatBuffer& GetDeltaBuffer();
 
 	private:
 		int32 BufferLength {0};
 		int32 LerpLength {0};
 		int32 CurrentLerpStep{0};
-		AlignedFloatBuffer DeltaBuffer;
+		FAlignedFloatBuffer DeltaBuffer;
 
 	}; // class BufferLerper
 }
