Fixing PVS 7.7 Issues:

warning V614: Uninitialized variable 'XYZ' used. Consider checking the fifth actual argument of the 'ZYX' function.

diff --git a/Engine/Plugins/Enterprise/DatasmithImporter/Source/DatasmithImporter/Private/Utility/DatasmithTextureResize.cpp b/Engine/Plugins/Enterprise/DatasmithImporter/Source/DatasmithImporter/Private/Utility/DatasmithTextureResize.cpp
--- a/Engine/Plugins/Enterprise/DatasmithImporter/Source/DatasmithImporter/Private/Utility/DatasmithTextureResize.cpp
+++ b/Engine/Plugins/Enterprise/DatasmithImporter/Source/DatasmithImporter/Private/Utility/DatasmithTextureResize.cpp
@@ -86,7 +86,7 @@ namespace
 	class NinePoints
 	{
 	public:
-		NinePoints(){ aa = ab = ac = ba = bb = bc = ca = cb = cc; }
+		NinePoints(){ aa = ab = ac = ba = bb = bc = ca = cb = cc = 0; }
 		float aa, ab, ac, ba, bb, bc, ca, cb, cc;
 	};
 
diff --git a/Engine/Plugins/Experimental/Enterprise/DatasmithRuntime/Source/Private/ImageReader.cpp b/Engine/Plugins/Experimental/Enterprise/DatasmithRuntime/Source/Private/ImageReader.cpp
--- a/Engine/Plugins/Experimental/Enterprise/DatasmithRuntime/Source/Private/ImageReader.cpp
+++ b/Engine/Plugins/Experimental/Enterprise/DatasmithRuntime/Source/Private/ImageReader.cpp
@@ -185,7 +185,7 @@ namespace DatasmithRuntime
 	class NinePoints
 	{
 	public:
-		NinePoints(){ aa = ab = ac = ba = bb = bc = ca = cb = cc; }
+		NinePoints(){ aa = ab = ac = ba = bb = bc = ca = cb = cc = 0; }
 		float aa, ab, ac, ba, bb, bc, ca, cb, cc;
 	};
 
diff --git a/Engine/Source/Runtime/SignalProcessing/Private/EnvelopeFollower.cpp b/Engine/Source/Runtime/SignalProcessing/Private/EnvelopeFollower.cpp
--- a/Engine/Source/Runtime/SignalProcessing/Private/EnvelopeFollower.cpp
+++ b/Engine/Source/Runtime/SignalProcessing/Private/EnvelopeFollower.cpp
@@ -25,7 +25,7 @@ namespace Audio
 
 	FEnvelopeFollower::FEnvelopeFollower(const float InSampleRate, const float InAttackTimeMsec, const float InReleaseTimeMSec, const EPeakMode::Type InMode, const bool bInIsAnalg)
 	{
-		Init(InSampleRate, InAttackTimeMsec, InReleaseTimeMSec, InMode, bIsAnalog);
+		Init(InSampleRate, InAttackTimeMsec, InReleaseTimeMSec, InMode, bInIsAnalg);
 	}
 
 	FEnvelopeFollower::~FEnvelopeFollower()
