Optimization to Audio::FastSin3
#rm miles.flanagan aaron.mcleran

diff --git a/Engine/Source/Runtime/SignalProcessing/Public/DSP/Dsp.h b/Engine/Source/Runtime/SignalProcessing/Public/DSP/Dsp.h
--- a/Engine/Source/Runtime/SignalProcessing/Public/DSP/Dsp.h
+++ b/Engine/Source/Runtime/SignalProcessing/Public/DSP/Dsp.h
@@ -76,15 +76,14 @@ namespace Audio
 		return X2;
 	}
 
+	// Valid on interval [-PI, PI] 
 	// Sine approximation using Bhaskara I technique discovered in 7th century. 
 	// https://en.wikipedia.org/wiki/Bh%C4%81skara_I
 	static FORCEINLINE float FastSin3(const float X)
 	{
-		// Component used to get negative radians
-		const float SafeX = X < 0.0f ? FMath::Min(X, -SMALL_NUMBER) : FMath::Max(X, SMALL_NUMBER);
-		const float Temp = SafeX * SafeX / FMath::Abs(SafeX);
-		const float Numerator = 16.0f * SafeX * (PI - Temp);
-		const float Denominator = 5.0f * PI * PI - 4.0f * Temp * (PI - Temp);
+		const float AbsX = FMath::Abs(X);
+		const float Numerator = 16.0f * X * (PI - AbsX);
+		const float Denominator = 5.0f * PI * PI - 4.0f * AbsX * (PI - AbsX);
 		return Numerator / Denominator;
 	}
 
