Making clear distinction between GetAbsoluteAverageAmplitude and GetAverageAmplitude in buffer vector ops.

diff --git a/Engine/Plugins/Runtime/AudioModulation/Source/AudioModulation/Private/Generators/SoundModulationEnvelopeFollower.cpp b/Engine/Plugins/Runtime/AudioModulation/Source/AudioModulation/Private/Generators/SoundModulationEnvelopeFollower.cpp
--- a/Engine/Plugins/Runtime/AudioModulation/Source/AudioModulation/Private/Generators/SoundModulationEnvelopeFollower.cpp
+++ b/Engine/Plugins/Runtime/AudioModulation/Source/AudioModulation/Private/Generators/SoundModulationEnvelopeFollower.cpp
@@ -77,7 +77,7 @@ namespace AudioModulation
 				static const int32 SimdMultiple = 4;
 				if (NumSamples % SimdMultiple == 0)
 				{
-					Amp = Audio::GetAverageAmplitude(TempBuffer.GetData(), NumSamples);
+					Amp = Audio::BufferGetAverageAbsValue(TempBuffer.GetData(), NumSamples);
 				}
 				// If NumSamples is not a multiple of 4, do it the non-SIMD way
 				else
@@ -85,7 +85,7 @@ namespace AudioModulation
 					Amp = 0.0f;
 					for (int32 i = 0; i < NumSamples; ++i)
 					{
-						Amp += TempBuffer[i];
+						Amp += FMath::Abs(TempBuffer[i]);
 					}
 					Amp /= NumSamples;
 				}
diff --git a/Engine/Source/Runtime/AudioMixer/Private/AudioMixerSourceManager.cpp b/Engine/Source/Runtime/AudioMixer/Private/AudioMixerSourceManager.cpp
--- a/Engine/Source/Runtime/AudioMixer/Private/AudioMixerSourceManager.cpp
+++ b/Engine/Source/Runtime/AudioMixer/Private/AudioMixerSourceManager.cpp
@@ -2295,7 +2295,7 @@ namespace Audio
 			if (!DisableEnvelopeFollowingCvar)
 			{
 				// Compute the source envelope using pre-distance attenuation buffer
-				float AverageSampleValue = Audio::GetAverageAmplitude(PreDistanceAttenBufferPtr, NumSamples);
+				float AverageSampleValue = Audio::BufferGetAverageAbsValue(PreDistanceAttenBufferPtr, NumSamples);
 				SourceInfo.SourceEnvelopeFollower.ProcessAudio(AverageSampleValue);
 
 				// Copy the current value of the envelope follower (block-rate value)
diff --git a/Engine/Source/Runtime/SignalProcessing/Private/BufferVectorOperations.cpp b/Engine/Source/Runtime/SignalProcessing/Private/BufferVectorOperations.cpp
--- a/Engine/Source/Runtime/SignalProcessing/Private/BufferVectorOperations.cpp
+++ b/Engine/Source/Runtime/SignalProcessing/Private/BufferVectorOperations.cpp
@@ -637,14 +637,14 @@ namespace Audio
 #endif
 	}
 
-	float GetAverageAmplitude(const FAlignedFloatBuffer& Buffer)
+	float BufferGetAverageValue(const FAlignedFloatBuffer& Buffer)
 	{
 		checkf(Buffer.Num() % 4 == 0, TEXT("Please use a buffer size that is a multiple of 4."));
 
-		return GetAverageAmplitude(Buffer.GetData(), Buffer.Num());
+		return BufferGetAverageValue(Buffer.GetData(), Buffer.Num());
 	}
 
-	float GetAverageAmplitude(const float* RESTRICT Buffer, int32 NumSamples)
+	float BufferGetAverageValue(const float* RESTRICT Buffer, int32 NumSamples)
 	{
 		checkf(NumSamples % 4 == 0, TEXT("Please use a buffer size that is a multiple of 4."));
 
@@ -671,6 +671,40 @@ namespace Audio
 #endif
 	}
 
+	float BufferGetAverageAbsValue(const FAlignedFloatBuffer& Buffer)
+	{
+		checkf(Buffer.Num() % 4 == 0, TEXT("Please use a buffer size that is a multiple of 4."));
+
+		return BufferGetAverageAbsValue(Buffer.GetData(), Buffer.Num());
+	}
+
+	float BufferGetAverageAbsValue(const float* RESTRICT Buffer, int32 NumSamples)
+	{
+		checkf(NumSamples % 4 == 0, TEXT("Please use a buffer size that is a multiple of 4."));
+
+#if !AUDIO_USE_SIMD
+		float Sum = 0.0f;
+		for (int32 i = 0; i < NumSamples; ++i)
+		{
+			Sum += FMath::Abs(Buffer[i]);
+		}
+		return Sum / NumSamples;
+#else
+		VectorRegister Sum = VectorZero();
+
+		for (int32 i = 0; i < NumSamples; i += 4)
+		{
+			VectorRegister Input = VectorAbs(VectorLoadAligned(&Buffer[i]));
+			Sum = VectorAdd(Sum, Input);
+		}
+
+		float PartionedSums[4];
+		VectorStore(Sum, PartionedSums);
+
+		return (PartionedSums[0] + PartionedSums[1] + PartionedSums[2] + PartionedSums[3]) / NumSamples;
+#endif
+	}
+
 	/**
 	 * CHANNEL MIXING OPERATIONS:
 	 * To understand these functions, it's best that you have prior experience reading SIMD code.
diff --git a/Engine/Source/Runtime/SignalProcessing/Public/DSP/AudioDebuggingUtilities.h b/Engine/Source/Runtime/SignalProcessing/Public/DSP/AudioDebuggingUtilities.h
--- a/Engine/Source/Runtime/SignalProcessing/Public/DSP/AudioDebuggingUtilities.h
+++ b/Engine/Source/Runtime/SignalProcessing/Public/DSP/AudioDebuggingUtilities.h
@@ -12,7 +12,7 @@ void BreakWhenAudible(float* InBuffer, int32 NumSamples)
 {
 	static const float AudibilityThreshold = Audio::ConvertToLinear(-40.0f);
 
-	float BufferAmplitude = Audio::GetAverageAmplitude(InBuffer, NumSamples);
+	float BufferAmplitude = Audio::BufferGetAverageAbsValue(InBuffer, NumSamples);
 
 	if (BufferAmplitude > AudibilityThreshold)
 	{
@@ -24,7 +24,7 @@ void BreakWhenTooLoud(float* InBuffer, int32 NumSamples)
 {
 	static const float PainThreshold = Audio::ConvertToLinear(3.0f);
 
-	float BufferAmplitude = Audio::GetAverageAmplitude(InBuffer, NumSamples);
+	float BufferAmplitude = Audio::BufferGetAverageAbsValue(InBuffer, NumSamples);
 
 	if (BufferAmplitude > PainThreshold)
 	{
diff --git a/Engine/Source/Runtime/SignalProcessing/Public/DSP/BufferVectorOperations.h b/Engine/Source/Runtime/SignalProcessing/Public/DSP/BufferVectorOperations.h
--- a/Engine/Source/Runtime/SignalProcessing/Public/DSP/BufferVectorOperations.h
+++ b/Engine/Source/Runtime/SignalProcessing/Public/DSP/BufferVectorOperations.h
@@ -123,9 +123,13 @@ namespace Audio
 	SIGNALPROCESSING_API float GetMagnitude(const FAlignedFloatBuffer& Buffer);
 	SIGNALPROCESSING_API float GetMagnitude(const float* RESTRICT Buffer, int32 NumSamples);
 
+	/** Takes an audio buffer and gets the average amplitude across that buffer. */
+	SIGNALPROCESSING_API float BufferGetAverageValue(const FAlignedFloatBuffer& Buffer);
+	SIGNALPROCESSING_API float BufferGetAverageValue(const float* RESTRICT Buffer, int32 NumSamples);
+
 	/** Takes an audio buffer and gets the average absolute amplitude across that buffer. */
-	SIGNALPROCESSING_API float GetAverageAmplitude(const FAlignedFloatBuffer& Buffer);
-	SIGNALPROCESSING_API float GetAverageAmplitude(const float* RESTRICT Buffer, int32 NumSamples);
+	SIGNALPROCESSING_API float BufferGetAverageAbsValue(const FAlignedFloatBuffer& Buffer);
+	SIGNALPROCESSING_API float BufferGetAverageAbsValue(const float* RESTRICT Buffer, int32 NumSamples);
 
 	/** CHANNEL-SPECIFIC OPERATIONS */
 
