dsp: adding fast accurate sin tone buffer generator

diff --git a/Engine/Source/Runtime/SignalProcessing/Public/DSP/Dsp.h b/Engine/Source/Runtime/SignalProcessing/Public/DSP/Dsp.h
--- a/Engine/Source/Runtime/SignalProcessing/Public/DSP/Dsp.h
+++ b/Engine/Source/Runtime/SignalProcessing/Public/DSP/Dsp.h
@@ -86,6 +86,160 @@ namespace Audio
 		const float Denominator = 5.0f * PI * PI - 4.0f * AbsX * (PI - AbsX);
 		return Numerator / Denominator;
 	}
+	
+	/**
+	* Generates a sine wave in the given buffer with the given frequency.
+	* 
+	* Uses the method of generating a 2D point on the unit circle and multiplying
+	* with a 2d rotation matrix to advance the phase. Very accurate and very fast.
+	* 
+	* Use this unless you're changing frequency every few samples on very tiny
+	* blocks (single digit sample count)
+	* 
+	* Distortion vs directly calling FMath::Sin() is on the order of -100db, and is roughly
+	* 10x faster (for constant frequency).
+	* 
+	* Frequency changes are not interpolated across the buffer, so dramatic changes
+	* could introduce aliasing due to a discontinuous derivative in the output.
+	* If you're changing frequencies by a lot, constantly, you should probably 
+	* use VectorSinCos in a loop.
+	* 
+	* Usage:
+	* 
+	* constexpr const int32 ChunkSampleCount = 480;
+	* alignas (sizeof(VectorRegister)) float ChunkBuffer[ChunkSampleCount];
+	* FSinOscBufferGenerator Generator;
+	* Generator.GenerateBuffer(48000, 400, ChunkBuffer, ChunkSampleCount);
+	* 
+	* ChunkBuffer now has 10ms of sine tone at 400hz. The next call will
+	* continue the sine tone where it left off.
+	*
+	*/
+	struct FSinOsc2DRotation
+	{
+		VectorRegister QuadDxVec, QuadDyVec;
+		float LastPhasePerSample;
+		float LastPhase;
+		float Dx, Dy;
+
+		FSinOsc2DRotation()
+		{
+			LastPhasePerSample = -1;
+			LastPhase = 0;
+		}
+
+		FSinOsc2DRotation(FSinOsc2DRotation const& other)
+		{
+			FMemory::Memcpy(this, &other, sizeof(FSinOsc2DRotation));
+		}
+
+		FSinOsc2DRotation& operator=(FSinOsc2DRotation const& other)
+		{
+			FMemory::Memcpy(this, &other, sizeof(FSinOsc2DRotation));
+			return *this;
+		}
+
+		/**
+		* Generates the sine tone, continuing from the last phase.
+		* 
+		* @param SampleRate the sample rate the buffer will be played at
+		* @param ClampedFrequency the frequency of the tone to emit, clamped to nyquist (SampleRate/2)
+		* @param AlignedBuffer the output buffer, aligned for vector operations
+		* @param BufferSampleCount the number of samples in the buffer. Generally, this should be %4 granularity.
+		* 
+		*/
+		void GenerateBuffer(float SampleRate, float ClampedFrequency, float* AlignedBuffer, int32 BufferSampleCount)
+		{
+			check(IsAligned(AlignedBuffer, sizeof(VectorRegister)));
+
+			// Regenerate our vector rotation components if our changes.
+			const float PhasePerSample = (ClampedFrequency * (2 * PI)) / (SampleRate);
+			if (LastPhasePerSample != PhasePerSample)
+			{
+				float QuadDx = FMath::Cos(PhasePerSample * 4);
+				float QuadDy = FMath::Sin(PhasePerSample * 4);
+
+				QuadDxVec = VectorLoadFloat1(&QuadDx);
+				QuadDyVec = VectorLoadFloat1(&QuadDy);
+
+				LastPhasePerSample = PhasePerSample;
+			}
+
+			float* Write = AlignedBuffer;
+
+			// The rotation matrix drifts, so we resync every so often
+			while (BufferSampleCount)
+			{
+				// The concept here is that cos/sin are points on a unit circle,
+				// so an oscilator is just a rotation of that point.
+				//
+				// To avoid drift, we resync off of an accurate sin/cos every evaluation
+				// then do a 2d rotation in the actual loop.
+				//
+				// we store 4 points at once to use SIMD, so each rotation is actually
+				// 4x the phase delta.
+				//
+				alignas(16) float PhaseSource[4];
+				PhaseSource[0] = LastPhase + 0 * PhasePerSample;
+				PhaseSource[1] = LastPhase + 1 * PhasePerSample;
+				PhaseSource[2] = LastPhase + 2 * PhasePerSample;
+				PhaseSource[3] = LastPhase + 3 * PhasePerSample;
+
+				VectorRegister PhaseVec = VectorLoadAligned(PhaseSource);
+				VectorRegister XVector, YVector;
+
+				// We need an accurate representation of the delta
+				// vectors since we are integrating it
+				VectorSinCos(&YVector, &XVector, &PhaseVec);
+
+				// Copy to local (the compiler actually didn't do this!)
+				VectorRegister LocalDxVec = QuadDxVec;
+				VectorRegister LocalDyVec = QuadDyVec;
+				
+				int32 BlockSampleCount = BufferSampleCount;
+				if (BlockSampleCount > 480)
+					BlockSampleCount = 480;
+
+				// Unrolling this didn't seem to really help perf wise.
+				int32 Block4 = BlockSampleCount >> 2;
+				while (Block4)
+				{
+					VectorStore(YVector, Write);
+
+					// 2D rotation matrix.
+					VectorRegister NewX = VectorSubtract(VectorMultiply(LocalDxVec, XVector), VectorMultiply(LocalDyVec, YVector));
+					VectorRegister NewY = VectorAdd(VectorMultiply(LocalDyVec, XVector), VectorMultiply(LocalDxVec, YVector));
+
+					XVector = NewX;
+					YVector = NewY;
+
+					Write += 4;
+					Block4--;
+				}
+
+				constexpr int32 SIMD_MASK = 0x00000003;
+				if (BlockSampleCount & SIMD_MASK)
+				{
+					// We've actually already calculated the next quad - it's in YVector
+					alignas(16) float YFloats[4];
+					VectorStoreAligned(YVector, YFloats);
+
+					int32 Remn = BlockSampleCount & SIMD_MASK;
+					for (int32 i = 0; i < Remn; i++)
+					{
+						Write[i] = YFloats[i];
+					}
+				}
+
+				// Advance phase, range reduce, and store.
+				float PhaseInRadians = LastPhase + BlockSampleCount * PhasePerSample;
+				PhaseInRadians -= FMath::FloorToFloat(PhaseInRadians / (2 * PI)) * (2 * PI);
+				LastPhase = PhaseInRadians;
+
+				BufferSampleCount -= BlockSampleCount;
+			} // end while BufferSampleCount
+		} // end GenerateBuffer
+	}; // end FSinOsc2DRotation
 
 	// Fast tanh based on pade approximation
 	static FORCEINLINE float FastTanh(float X)
