Metasound: Midi Note Quantization node
#jria UEAU-769, UEAU-761

diff --git a/Engine/Plugins/Runtime/Metasound/Source/MetasoundStandardNodes/Private/MetasoundMIDIToFreqNode.cpp b/Engine/Plugins/Runtime/Metasound/Source/MetasoundStandardNodes/Private/MetasoundMIDIToFreqNode.cpp
--- a/Engine/Plugins/Runtime/Metasound/Source/MetasoundStandardNodes/Private/MetasoundMIDIToFreqNode.cpp
+++ b/Engine/Plugins/Runtime/Metasound/Source/MetasoundStandardNodes/Private/MetasoundMIDIToFreqNode.cpp
@@ -6,6 +6,7 @@
 #include "MetasoundNodeRegistrationMacro.h"
 #include "MetasoundPrimitives.h"
 #include "MetasoundStandardNodesNames.h"
+#include "MetasoundStandardNodesCategories.h"
 #include "DSP/Dsp.h"
 #include "MetasoundParamHelper.h"
 
@@ -161,6 +162,7 @@ namespace Metasound
 			Info.Author = PluginAuthor;
 			Info.PromptIfMissing = PluginNodeMissingPrompt;
 			Info.DefaultInterface = GetVertexInterface();
+			Info.CategoryHierarchy.Emplace(StandardNodes::Music);
 
 			return Info;
 		};
diff --git a/Engine/Plugins/Runtime/Metasound/Source/MetasoundStandardNodes/Private/MetasoundMidiNoteQuantizerNode.cpp b/Engine/Plugins/Runtime/Metasound/Source/MetasoundStandardNodes/Private/MetasoundMidiNoteQuantizerNode.cpp
new file mode 100644
--- /dev/null
+++ b/Engine/Plugins/Runtime/Metasound/Source/MetasoundStandardNodes/Private/MetasoundMidiNoteQuantizerNode.cpp
@@ -0,0 +1,240 @@
+// Copyright Epic Games, Inc. All Rights Reserved.
+
+#include "CoreMinimal.h"
+#include "DSP/MidiNoteQuantizer.h"
+#include "Internationalization/Text.h"
+
+#include "MetasoundFacade.h"
+#include "MetasoundPrimitives.h"
+#include "MetasoundAudioBuffer.h"
+#include "MetasoundParamHelper.h"
+#include "MetasoundOperatorSettings.h"
+#include "MetasoundStandardNodesNames.h"
+#include "MetasoundExecutableOperator.h"
+#include "MetasoundNodeRegistrationMacro.h"
+#include "MetasoundStandardNodesCategories.h"
+#include "MetasoundDataTypeRegistrationMacro.h"
+
+#define LOCTEXT_NAMESPACE "MetasoundBasicFilterNodes"
+
+namespace Metasound
+{
+	// forward declarations
+	// ...
+		
+#pragma region Parameter Names
+	namespace MidiNoteQuantizerParameterNames
+	{
+		// inputs
+		METASOUND_PARAM(ParamNoteIn, "Note In", "Midi Note to quantize");
+		METASOUND_PARAM(ParamRootNote, "Root Note", "Midi note to treat as the Root (0 = C, 1 = D, etc.).  Octave does not matter. Values < 0 will clamp to 0");
+		METASOUND_PARAM(ParamScaleDegrees, "Scale Degrees"
+			, "Set of notes in ascending order, represeting half steps starting at 0.0f meaning the Root Note. The highest value MUST be the same as 0.0 one octave higher: 12.0 in most cases, but it doesn't have to be (i.e. could be 24.0 if you want to define a 2 octave range)");
+
+		// outputs
+		METASOUND_PARAM(ParamNoteOutput, "Note Out", "Quantized Note");
+
+	} // namespace MidiNoteQuantizerParameterNames
+
+	using namespace MidiNoteQuantizerParameterNames;
+
+#pragma endregion // Parameter Names
+
+
+#pragma region Operator Declaration
+	class FMidiNoteQuantizerOperator: public TExecutableOperator < FMidiNoteQuantizerOperator >
+	{
+	public:
+		using FArrayScaleDegreeReadRef = TDataReadReference<TArray<float>>;
+		using ScaleDegreeArrayType = TArray<float>;
+
+		// ctor
+		FMidiNoteQuantizerOperator(
+			  const FOperatorSettings& InSettings
+			, const FFloatReadRef& InMidiNoteIn
+			, const FFloatReadRef& InRootNote
+			, const FArrayScaleDegreeReadRef& InScale
+		);
+
+		// node interface
+		static const FNodeClassMetadata& GetNodeInfo();
+		static FVertexInterface DeclareVertexInterface();
+		static TUniquePtr<IOperator> CreateOperator(const FCreateOperatorParams& InParams, FBuildErrorArray& OutErrors);
+		virtual FDataReferenceCollection GetInputs() const override;
+		virtual FDataReferenceCollection GetOutputs() const override;
+		void Execute();
+
+	private: // members
+		// input pins
+		FFloatReadRef MidiNoteIn;
+		FFloatReadRef RootNote;
+		FArrayScaleDegreeReadRef Scale;
+
+		// output pins
+		FFloatWriteRef MidiNoteOut;
+
+		// cached values
+		float PreviousNoteIn = TNumericLimits<float>::Lowest();
+		float PreviousRoot = -1.0f;
+		TArray<float> PreviousScale;
+		float PreviousNoteOut = 0.0f;
+
+		// other
+		FOperatorSettings Settings;
+
+	}; // class FMidiNoteQuantizerOperator
+
+#pragma endregion // Operator Declaration
+
+
+#pragma region Operator Implementation
+
+	// ctor
+	FMidiNoteQuantizerOperator::FMidiNoteQuantizerOperator(
+		  const FOperatorSettings& InSettings
+		, const FFloatReadRef& InMidiNoteIn
+		, const FFloatReadRef& InRootNote
+		, const FArrayScaleDegreeReadRef& InScale
+		)
+		: MidiNoteIn(InMidiNoteIn)
+		, RootNote(InRootNote)
+		, Scale(InScale)
+		, MidiNoteOut(FFloatWriteRef::CreateNew())
+		, Settings(InSettings)
+	{
+	}
+
+
+	const FNodeClassMetadata& FMidiNoteQuantizerOperator::GetNodeInfo()
+	{
+		auto InitNodeInfo = []() -> FNodeClassMetadata
+		{
+			FNodeClassMetadata Info;
+			Info.ClassName = { Metasound::StandardNodes::Namespace, TEXT("MIDI Note Quantizer"), Metasound::StandardNodes::AudioVariant };
+			Info.MajorVersion = 1;
+			Info.MinorVersion = 0;
+			Info.DisplayName = LOCTEXT("Metasound_MIDI_Note_Quantizer_NodeDisplayName", "MIDI Note Quantizer");
+			Info.Description = LOCTEXT("MidiNoteQuantizer_NodeDescription", "Quantizes a MIDI note to the nearset note that matches provided criteria"),
+			Info.Author = PluginAuthor;
+			Info.PromptIfMissing = PluginNodeMissingPrompt;
+			Info.DefaultInterface = DeclareVertexInterface();
+			Info.CategoryHierarchy.Emplace(StandardNodes::Music);
+
+			return Info;
+		};
+
+		static const FNodeClassMetadata Info = InitNodeInfo();
+
+		return Info;
+	}
+
+
+	FVertexInterface FMidiNoteQuantizerOperator::DeclareVertexInterface()
+	{
+		static const FVertexInterface Interface(
+			FInputVertexInterface(
+				TInputDataVertexModel<float>(METASOUND_GET_PARAM_NAME(ParamNoteIn), METASOUND_GET_PARAM_TT(ParamNoteIn), 60.0f),
+				TInputDataVertexModel<float>(METASOUND_GET_PARAM_NAME(ParamRootNote), METASOUND_GET_PARAM_TT(ParamRootNote), 0.0f),
+				TInputDataVertexModel<ScaleDegreeArrayType>(METASOUND_GET_PARAM_NAME(ParamScaleDegrees), METASOUND_GET_PARAM_TT(ParamScaleDegrees))
+			),
+			FOutputVertexInterface(
+				TOutputDataVertexModel<float>(METASOUND_GET_PARAM_NAME(ParamNoteOutput), METASOUND_GET_PARAM_TT(ParamNoteOutput))
+				)
+			);
+
+		return Interface;
+	}
+
+
+	TUniquePtr<IOperator> FMidiNoteQuantizerOperator::CreateOperator(const FCreateOperatorParams& InParams, FBuildErrorArray& OutErrors)
+	{
+		const FDataReferenceCollection& InputDataRefs = InParams.InputDataReferences;
+
+		// inputs
+		FFloatReadRef MidiNoteIn = InputDataRefs.GetDataReadReferenceOrConstruct<float>(METASOUND_GET_PARAM_NAME(ParamNoteIn));
+		FFloatReadRef RootNoteIn = InputDataRefs.GetDataReadReferenceOrConstruct<float>(METASOUND_GET_PARAM_NAME(ParamRootNote));
+		FArrayScaleDegreeReadRef InScaleArray = InParams.InputDataReferences.GetDataReadReferenceOrConstruct<ScaleDegreeArrayType>(METASOUND_GET_PARAM_NAME(ParamScaleDegrees));
+
+		return MakeUnique <FMidiNoteQuantizerOperator>(
+			  InParams.OperatorSettings
+			, MidiNoteIn
+			, RootNoteIn
+			, InScaleArray
+			);
+	}
+
+	FDataReferenceCollection FMidiNoteQuantizerOperator::GetInputs() const
+	{
+		FDataReferenceCollection InputDataReferences;
+		InputDataReferences.AddDataReadReference(METASOUND_GET_PARAM_NAME(ParamNoteIn), FFloatReadRef(MidiNoteIn));
+		InputDataReferences.AddDataReadReference(METASOUND_GET_PARAM_NAME(ParamRootNote), FFloatReadRef(RootNote));
+		InputDataReferences.AddDataReadReference(METASOUND_GET_PARAM_NAME(ParamScaleDegrees), FArrayScaleDegreeReadRef(Scale));
+
+		return InputDataReferences;
+	}
+
+	FDataReferenceCollection FMidiNoteQuantizerOperator::GetOutputs() const
+	{
+		// expose read access to our output buffer for other processors in the graph
+		FDataReferenceCollection OutputDataReferences;
+		OutputDataReferences.AddDataReadReference(METASOUND_GET_PARAM_NAME(ParamNoteOutput), FFloatReadRef(MidiNoteOut));
+
+		return OutputDataReferences;
+	}
+
+	void FMidiNoteQuantizerOperator::Execute()
+	{
+		if ((*Scale).IsEmpty())
+		{
+			*MidiNoteOut = *MidiNoteIn;
+		}
+
+		// calculate new output and cache values if needed
+		const float CurrentRoot = FMath::Max(*RootNote, 0.0f);
+		if (!FMath::IsNearlyEqual(PreviousNoteIn, *MidiNoteIn)
+			|| !FMath::IsNearlyEqual(PreviousRoot, CurrentRoot)
+			|| PreviousScale !=  *Scale
+			)
+		{
+			// cache values
+			PreviousNoteIn = *MidiNoteIn;
+			PreviousRoot = *RootNote;
+			PreviousScale = *Scale;
+
+			PreviousScale.Sort();
+			PreviousNoteOut = Audio::FMidiNoteQuantizer::QuantizeMidiNote(PreviousNoteIn, PreviousRoot, PreviousScale);
+		}
+
+		// set the output value
+		*MidiNoteOut = PreviousNoteOut;
+	}
+
+#pragma endregion // Operator Implementation
+
+
+#pragma region Node Declaration
+	class METASOUNDSTANDARDNODES_API FMidiNoteQuantizerNode: public FNodeFacade
+	{
+	public:
+		// public node api needs to define two conversion constructors:
+		// (1: from FString)
+		FMidiNoteQuantizerNode(const FString & InInstanceName, const FGuid & InInstanceID)
+			: FNodeFacade(InInstanceName, InInstanceID, TFacadeOperatorClass <FMidiNoteQuantizerOperator>())
+		{ }
+
+		// (2: From an NodeInitData struct)
+		FMidiNoteQuantizerNode(const FNodeInitData & InInitData)
+			: FMidiNoteQuantizerNode(InInitData.InstanceName, InInitData.InstanceID)
+		{ }
+
+	};
+#pragma endregion // Node Declaration
+
+
+#pragma region Node Registration
+	METASOUND_REGISTER_NODE(FMidiNoteQuantizerNode);
+#pragma endregion // Node Registration
+
+} // namespace Metasound
+
+#undef LOCTEXT_NAMESPACE //MetasoundBasicFilterNodes
\ No newline at end of file
diff --git a/Engine/Plugins/Runtime/Metasound/Source/MetasoundStandardNodes/Private/MetasoundMidiScaleToArrayNode.cpp b/Engine/Plugins/Runtime/Metasound/Source/MetasoundStandardNodes/Private/MetasoundMidiScaleToArrayNode.cpp
new file mode 100644
--- /dev/null
+++ b/Engine/Plugins/Runtime/Metasound/Source/MetasoundStandardNodes/Private/MetasoundMidiScaleToArrayNode.cpp
@@ -0,0 +1,213 @@
+// Copyright Epic Games, Inc. All Rights Reserved.
+
+#include "CoreMinimal.h"
+#include "DSP/MidiNoteQuantizer.h"
+#include "Internationalization/Text.h"
+
+#include "MetasoundFacade.h"
+#include "MetasoundPrimitives.h"
+#include "MetasoundScaleEnums.h"
+#include "MetasoundAudioBuffer.h"
+#include "MetasoundParamHelper.h"
+#include "MetasoundOperatorSettings.h"
+#include "MetasoundStandardNodesNames.h"
+#include "MetasoundExecutableOperator.h"
+#include "MetasoundNodeRegistrationMacro.h"
+#include "MetasoundStandardNodesCategories.h"
+#include "MetasoundDataTypeRegistrationMacro.h"
+
+#define LOCTEXT_NAMESPACE "MetasoundBasicFilterNodes"
+
+namespace Metasound
+{
+#pragma region Parameter Names
+	namespace MusicalScaleToNoteArrayParameterNames
+	{
+		// inputs
+		METASOUND_PARAM(ParamScaleDegreesPreset, "Scale Degrees", "Select scale preset");
+		METASOUND_PARAM(ParamChordTonesOnly, "Chord Tones Only", "If true, will only return a subset of the scale represeting chord tones. (i.e. scale degrees 1,3,5,7)");
+
+		// outputs
+		METASOUND_PARAM(ParamNoteArrayOutput, "Scale Array Out", "Array represeting the scale as half steps above the root. The set is inclusive at both ends: (starting at 0.0f and ending with 12.0f)");
+
+	} // namespace MusicalScaleToNoteArrayParameterNames
+
+	using namespace MusicalScaleToNoteArrayParameterNames;
+
+#pragma endregion // Parameter Names
+
+
+#pragma region Operator Declaration
+	class FMusicalScaleToNoteArrayOperator: public TExecutableOperator < FMusicalScaleToNoteArrayOperator >
+	{
+	public:
+		using FArrayScaleDegreeWriteRef = TDataWriteReference<TArray<float>>;
+		using ScaleDegreeArrayType = TArray<float>;
+
+		// ctor
+		FMusicalScaleToNoteArrayOperator(
+			  const FOperatorSettings& InSettings
+			, const FEnumMusicalScaleReadRef& InScale
+			, const FBoolReadRef& InChordTonesOnly
+		);
+
+		// node interface
+		static const FNodeClassMetadata& GetNodeInfo();
+		static FVertexInterface DeclareVertexInterface();
+		static TUniquePtr<IOperator> CreateOperator(const FCreateOperatorParams& InParams, FBuildErrorArray& OutErrors);
+		virtual FDataReferenceCollection GetInputs() const override;
+		virtual FDataReferenceCollection GetOutputs() const override;
+		void Execute();
+
+	private: // members
+		// input pins
+		FEnumMusicalScaleReadRef Scale;
+		FBoolReadRef bChordTonesOnly;
+
+		// output pins
+		FArrayScaleDegreeWriteRef ScaleArrayOutput;
+
+		// cached values
+		Audio::EMusicalScale::Scale PreviousScale = Audio::EMusicalScale::Scale::Count;
+		bool bPreviousChordTones = false;
+		TArray<float> PreviousNoteOut;
+
+		// other
+		FOperatorSettings Settings;
+
+	}; // class FMusicalScaleToNoteArrayOperator
+
+#pragma endregion // Operator Declaration
+
+
+#pragma region Operator Implementation
+
+	// ctor
+	FMusicalScaleToNoteArrayOperator::FMusicalScaleToNoteArrayOperator(
+		  const FOperatorSettings& InSettings
+		, const FEnumMusicalScaleReadRef& InScale
+		, const FBoolReadRef& InChordTonesOnly
+		)
+		: Scale(InScale)
+		, bChordTonesOnly(InChordTonesOnly)
+		, ScaleArrayOutput(FArrayScaleDegreeWriteRef::CreateNew())
+		, Settings(InSettings)
+	{
+	}
+
+
+	const FNodeClassMetadata& FMusicalScaleToNoteArrayOperator::GetNodeInfo()
+	{
+		auto InitNodeInfo = []() -> FNodeClassMetadata
+		{
+			FNodeClassMetadata Info;
+			Info.ClassName = { Metasound::StandardNodes::Namespace, TEXT("Musical Scale To Note Array"), Metasound::StandardNodes::AudioVariant };
+			Info.MajorVersion = 1;
+			Info.MinorVersion = 0;
+			Info.DisplayName = LOCTEXT("Metasound_ScaleToNoteArray_NodeDisplayName", "Scale to Note Array");
+			Info.Description = LOCTEXT("MusicalScaleToNoteArray_NodeDescription", "Returns an array represeting the selected scale"),
+			Info.Author = PluginAuthor;
+			Info.PromptIfMissing = PluginNodeMissingPrompt;
+			Info.DefaultInterface = DeclareVertexInterface();
+			Info.CategoryHierarchy.Emplace(StandardNodes::Music);
+
+			return Info;
+		};
+
+		static const FNodeClassMetadata Info = InitNodeInfo();
+
+		return Info;
+	}
+
+
+	FVertexInterface FMusicalScaleToNoteArrayOperator::DeclareVertexInterface()
+	{
+		static const FVertexInterface Interface(
+			FInputVertexInterface(
+				TInputDataVertexModel<FEnumEMusicalScale>(METASOUND_GET_PARAM_NAME(ParamScaleDegreesPreset), METASOUND_GET_PARAM_TT(ParamScaleDegreesPreset)),
+				TInputDataVertexModel<bool>(METASOUND_GET_PARAM_NAME(ParamChordTonesOnly), METASOUND_GET_PARAM_TT(ParamChordTonesOnly))
+			),
+			FOutputVertexInterface(
+				TOutputDataVertexModel<ScaleDegreeArrayType>(METASOUND_GET_PARAM_NAME(ParamNoteArrayOutput), METASOUND_GET_PARAM_TT(ParamNoteArrayOutput))
+				)
+			);
+
+		return Interface;
+	}
+
+
+	TUniquePtr<IOperator> FMusicalScaleToNoteArrayOperator::CreateOperator(const FCreateOperatorParams& InParams, FBuildErrorArray& OutErrors)
+	{
+		const FDataReferenceCollection& InputDataRefs = InParams.InputDataReferences;
+
+		// inputs
+		FEnumMusicalScaleReadRef Scale = InputDataRefs.GetDataReadReferenceOrConstruct<FEnumEMusicalScale>(METASOUND_GET_PARAM_NAME(ParamScaleDegreesPreset));
+		FBoolReadRef ChordTonesOnly = InputDataRefs.GetDataReadReferenceOrConstruct<bool>(METASOUND_GET_PARAM_NAME(ParamChordTonesOnly));
+
+		return MakeUnique <FMusicalScaleToNoteArrayOperator>(
+			  InParams.OperatorSettings
+			, Scale
+			, ChordTonesOnly
+			);
+	}
+
+	FDataReferenceCollection FMusicalScaleToNoteArrayOperator::GetInputs() const
+	{
+		FDataReferenceCollection InputDataReferences;
+		InputDataReferences.AddDataReadReference(METASOUND_GET_PARAM_NAME(ParamScaleDegreesPreset), FEnumMusicalScaleReadRef(Scale));
+		InputDataReferences.AddDataReadReference(METASOUND_GET_PARAM_NAME(ParamChordTonesOnly), FBoolReadRef(bChordTonesOnly));
+
+		return InputDataReferences;
+	}
+
+	FDataReferenceCollection FMusicalScaleToNoteArrayOperator::GetOutputs() const
+	{
+		// expose read access to our output buffer for other processors in the graph
+		FDataReferenceCollection OutputDataReferences;
+		OutputDataReferences.AddDataReadReference(METASOUND_GET_PARAM_NAME(ParamNoteArrayOutput), FArrayScaleDegreeWriteRef(ScaleArrayOutput));
+
+		return OutputDataReferences;
+	}
+
+	void FMusicalScaleToNoteArrayOperator::Execute()
+	{
+		// calculate new output and cache values if needed
+		if ((PreviousScale != *Scale) || (bPreviousChordTones != *bChordTonesOnly))
+		{
+			// cache values
+			PreviousScale = *Scale;
+			bPreviousChordTones = *bChordTonesOnly;
+
+			PreviousNoteOut = *ScaleArrayOutput = Audio::FMidiNoteQuantizer::ScaleDegreeSetMap[PreviousScale].GetScaleDegreeSet(bPreviousChordTones);
+		}
+	}
+
+#pragma endregion // Operator Implementation
+
+
+#pragma region Node Declaration
+	class METASOUNDSTANDARDNODES_API FMusicalScaleToNoteArrayNode: public FNodeFacade
+	{
+	public:
+		// public node api needs to define two conversion constructors:
+		// (1: from FString)
+		FMusicalScaleToNoteArrayNode(const FString & InInstanceName, const FGuid & InInstanceID)
+			: FNodeFacade(InInstanceName, InInstanceID, TFacadeOperatorClass <FMusicalScaleToNoteArrayOperator>())
+		{ }
+
+		// (2: From an NodeInitData struct)
+		FMusicalScaleToNoteArrayNode(const FNodeInitData & InInitData)
+			: FMusicalScaleToNoteArrayNode(InInitData.InstanceName, InInitData.InstanceID)
+		{ }
+
+	};
+#pragma endregion // Node Declaration
+
+
+#pragma region Node Registration
+	METASOUND_REGISTER_NODE(FMusicalScaleToNoteArrayNode);
+#pragma endregion // Node Registration
+
+} // namespace Metasound
+
+#undef LOCTEXT_NAMESPACE //MetasoundBasicFilterNodes
\ No newline at end of file
diff --git a/Engine/Plugins/Runtime/Metasound/Source/MetasoundStandardNodes/Private/MetasoundStandardNodesCategories.cpp b/Engine/Plugins/Runtime/Metasound/Source/MetasoundStandardNodes/Private/MetasoundStandardNodesCategories.cpp
--- a/Engine/Plugins/Runtime/Metasound/Source/MetasoundStandardNodes/Private/MetasoundStandardNodesCategories.cpp
+++ b/Engine/Plugins/Runtime/Metasound/Source/MetasoundStandardNodes/Private/MetasoundStandardNodesCategories.cpp
@@ -11,6 +11,7 @@ namespace Metasound
 		const FText Generators = { LOCTEXT("Metasound_GeneratorsCategory", "Generators") };
 		const FText Filters = { LOCTEXT("Metasound_FiltersCategory", "Filters") };
 		const FText TriggerUtils = { LOCTEXT("Metasound_TriggerUtilsCategory", "Trigger Utilities") };
+		const FText Music = { LOCTEXT("Metasound_MusicCategory", "Music") };
 		const FText RandomUtils = { LOCTEXT("Metasound_RandomUtilsCategory", "Random Utilities") };
 		const FText Conversions  = { LOCTEXT("Metasound_ConversionsCategory", "Conversions") };
 		const FText DebugUtils = { LOCTEXT("Metasound_DebugUtilsCategory", "Debug Utilities") };
diff --git a/Engine/Plugins/Runtime/Metasound/Source/MetasoundStandardNodes/Public/MetasoundScaleEnums.h b/Engine/Plugins/Runtime/Metasound/Source/MetasoundStandardNodes/Public/MetasoundScaleEnums.h
new file mode 100644
--- /dev/null
+++ b/Engine/Plugins/Runtime/Metasound/Source/MetasoundStandardNodes/Public/MetasoundScaleEnums.h
@@ -0,0 +1,64 @@
+// Copyright Epic Games, Inc. All Rights Reserved.
+
+#pragma once
+
+#include "CoreMinimal.h"
+#include "MetasoundEnum.h"
+#include "DSP/MidiNoteQuantizer.h"
+#include "Internationalization/Text.h"
+#include "MetasoundDataTypeRegistrationMacro.h"
+
+namespace Metasound
+{
+#define LOCTEXT_NAMESPACE "MidiScaleDefinitions"
+	// Any desired additions to this Enum/list need to first be added to the EMusicalScale enum in MidiNoteQuantizer.h
+	// and defined in MidiNoteQuantizer.cpp in the TMap<EMusicalScale::Scale, ScaleDegreeSet> ScaleDegreeSetMap static init
+
+	// Metasound enum
+	DECLARE_METASOUND_ENUM(Audio::EMusicalScale::Scale, Audio::EMusicalScale::Scale::Major,
+	METASOUNDSTANDARDNODES_API, FEnumEMusicalScale, FEnumMusicalScaleTypeInfo, FEnumMusicalScaleReadRef, FEnumMusicalScaleWriteRef);
+	DEFINE_METASOUND_ENUM_BEGIN(Audio::EMusicalScale::Scale, FEnumEMusicalScale, "MusicalScale")
+
+	// modes
+	DEFINE_METASOUND_ENUM_ENTRY(Audio::EMusicalScale::Scale::Major, LOCTEXT("MajorDescription", "Major Scale"), LOCTEXT("MajorDescriptionTT", "Major (Ionian)")),
+	DEFINE_METASOUND_ENUM_ENTRY(Audio::EMusicalScale::Scale::Minor_Dorian, LOCTEXT("Minor_DorianDescription", "Minor (Dorian)"), LOCTEXT("Minor_DorianDescriptionTT", "Dorian Minor")),
+	DEFINE_METASOUND_ENUM_ENTRY(Audio::EMusicalScale::Scale::Phrygian, LOCTEXT("PhrygianDescription", "Phrygian "), LOCTEXT("PhrygianDescriptionTT", "Phrygian")),
+	DEFINE_METASOUND_ENUM_ENTRY(Audio::EMusicalScale::Scale::Lydian, LOCTEXT("LydianDescription", "Lydian"), LOCTEXT("LydianDescriptionTT", "Lydian (sharp-4)")),
+	DEFINE_METASOUND_ENUM_ENTRY(Audio::EMusicalScale::Scale::Dominant7th_Mixolydian, LOCTEXT("Dominant7th_MixolydianDescription", "Dominant 7th (Mixolydian)"), LOCTEXT("Dominant7th_MixolydianDescriptionTT", "Mioxlydian (Dominant 7)")),
+	DEFINE_METASOUND_ENUM_ENTRY(Audio::EMusicalScale::Scale::NaturalMinor_Aeolian, LOCTEXT("NaturalMinor_AeolianDescription", "Natural Minor (Aeolian)"), LOCTEXT("NaturalMinor_AeolianDescriptionTT", "Natural Minor (Aeolian)")),
+	DEFINE_METASOUND_ENUM_ENTRY(Audio::EMusicalScale::Scale::HalfDiminished_Locrian, LOCTEXT("HalfDiminished_LocrianDescription", "Half Diminished (Locrian)"), LOCTEXT("HalfDiminished_LocrianDescriptionTT", "Half-Diminished (Locrian)")),
+	DEFINE_METASOUND_ENUM_ENTRY(Audio::EMusicalScale::Scale::Diminished, LOCTEXT("DiminishedDescription", "Diminished "), LOCTEXT("DiminishedDescriptionTT", "Diminished")),
+	// non-diatonic
+	DEFINE_METASOUND_ENUM_ENTRY(Audio::EMusicalScale::Scale::Chromatic, LOCTEXT("ChromaticDescription", "Chromatic"), LOCTEXT("ChromaticDescriptionTT", "Chromatic")),
+	DEFINE_METASOUND_ENUM_ENTRY(Audio::EMusicalScale::Scale::WholeTone, LOCTEXT("WholeToneDescription", "Whole-Tone"), LOCTEXT("WholeToneDescriptionTT", "Whole Tone")),
+	DEFINE_METASOUND_ENUM_ENTRY(Audio::EMusicalScale::Scale::DiminishedWholeTone, LOCTEXT("DiminishedWholeToneDescription", "Diminished Whole-Tone"), LOCTEXT("DiminishedWholeToneDescriptionTT", "Diminished Whole Tone")),
+	// petantonic
+	DEFINE_METASOUND_ENUM_ENTRY(Audio::EMusicalScale::Scale::MajorPentatonic, LOCTEXT("MajorPentatonicDescription", "Major Pentatonic "), LOCTEXT("MajorPentatonicDescriptionTT", "Major Pentatonic")),
+	DEFINE_METASOUND_ENUM_ENTRY(Audio::EMusicalScale::Scale::MinorPentatonic, LOCTEXT("MinorPentatonicDescription", "Minor Pentatonic "), LOCTEXT("MinorPentatonicDescriptionTT", "Minor Pentatonic")),
+	DEFINE_METASOUND_ENUM_ENTRY(Audio::EMusicalScale::Scale::Blues, LOCTEXT("BluesDescription", "Blues "), LOCTEXT("BluesDescriptionTT", "Blues")),
+	// bebop
+	DEFINE_METASOUND_ENUM_ENTRY(Audio::EMusicalScale::Scale::Bebop_Major, LOCTEXT("Bebop_MajorDescription", "Bebop (Major)"), LOCTEXT("Bebop_MajorDescriptionTT", "Bebop Major")),
+	DEFINE_METASOUND_ENUM_ENTRY(Audio::EMusicalScale::Scale::Bebop_Minor, LOCTEXT("Bebop_MinorDescription", "Bebop (Minor)"), LOCTEXT("Bebop_MinorDescriptionTT", "Bebop Minor")),
+	DEFINE_METASOUND_ENUM_ENTRY(Audio::EMusicalScale::Scale::Bebop_MinorNumber2, LOCTEXT("Bebop_MinorNumber2Description", "Bebop (Minor) #2"), LOCTEXT("Bebop_MinorNumber2DescriptionTT", "Bebop Minor #2")),
+	DEFINE_METASOUND_ENUM_ENTRY(Audio::EMusicalScale::Scale::Bebop_Dominant, LOCTEXT("Bebop_DominantDescription", "Bebop (Dominant)"), LOCTEXT("Bebop_DominantDescriptionTT", "Bebop Dominant")),
+	// common major/minors
+	DEFINE_METASOUND_ENUM_ENTRY(Audio::EMusicalScale::Scale::HarmonicMajor, LOCTEXT("HarmonicMajorDescription", "Harmonic Major"), LOCTEXT("HarmonicMajorDescriptionTT", "Harmonic Major")),
+	DEFINE_METASOUND_ENUM_ENTRY(Audio::EMusicalScale::Scale::HarmonicMinor, LOCTEXT("HarmonicMinorDescription", "Harmonic Minor "), LOCTEXT("HarmonicMinorDescriptionTT", "Harmonic Minor")),
+	DEFINE_METASOUND_ENUM_ENTRY(Audio::EMusicalScale::Scale::MelodicMinor, LOCTEXT("MelodicMinorDescription", "Melodic Minor "), LOCTEXT("MelodicMinorDescriptionTT", "Melodic Minor")),
+	DEFINE_METASOUND_ENUM_ENTRY(Audio::EMusicalScale::Scale::SixthModeOfHarmonicMinor, LOCTEXT("SixthModeOfHarmonicMinorDescription", "Sixth Mode of Harmonic Minor"), LOCTEXT("SixthModeOfHarmonicMinorDescriptionTT", "Sixth Mode of Harmonic Minor")),
+	// lydian/augmented
+	DEFINE_METASOUND_ENUM_ENTRY(Audio::EMusicalScale::Scale::LydianAugmented, LOCTEXT("LydianAugmentedDescription", "Lydian Augmented"), LOCTEXT("LydianAugmentedDescriptionTT", "Lydian Augmented")),
+	DEFINE_METASOUND_ENUM_ENTRY(Audio::EMusicalScale::Scale::LydianDominant, LOCTEXT("LydianDominantDescription", "Lydian Dominant "), LOCTEXT("LydianDominantDescriptionTT", "Lydian Dominant")),
+	DEFINE_METASOUND_ENUM_ENTRY(Audio::EMusicalScale::Scale::Augmented, LOCTEXT("AugmentedDescription", "Augmented"), LOCTEXT("AugmentedDescriptionTT", "Augmented")),
+	// diminished
+	DEFINE_METASOUND_ENUM_ENTRY(Audio::EMusicalScale::Scale::Diminished_BeginWithHalfStep, LOCTEXT("Diminished_BeginWithHalfStepDescription", "Diminished (Begin With Half-Step)"), LOCTEXT("Diminished_BeginWithHalfStepDescriptionTT", "Diminished (begins with Half Step)")),
+	DEFINE_METASOUND_ENUM_ENTRY(Audio::EMusicalScale::Scale::Diminished_BeginWithWholeStep, LOCTEXT("Diminished_BeginWithWholeStepDescription", "Diminished (Begin With Whole-Step"), LOCTEXT("Diminished_BeginWithWholeStepDescriptionTT", "Diminished (begins with Whole Step)")),
+	DEFINE_METASOUND_ENUM_ENTRY(Audio::EMusicalScale::Scale::HalfDiminished_LocrianNumber2, LOCTEXT("HalfDiminished_LocrianNumber2Description", "Half-Diminished (Locrian #2)"), LOCTEXT("HalfDiminished_LocrianNumber2DescriptionTT", "Half Diminished Locrian (#2)")),
+	// other
+	DEFINE_METASOUND_ENUM_ENTRY(Audio::EMusicalScale::Scale::Spanish_or_Jewish, LOCTEXT("Spanish_or_JewishDescription", "Spanish or Jewish Scale"), LOCTEXT("Spanish_or_JewishDescriptionTT", "Spanish/Jewish")),
+	DEFINE_METASOUND_ENUM_ENTRY(Audio::EMusicalScale::Scale::Hindu, LOCTEXT("HinduDescription", "Hindu "), LOCTEXT("HinduDescriptionTT", "Hindu"))
+
+	DEFINE_METASOUND_ENUM_END()
+
+#undef LOCTEXT_NAMESPACE
+} // namespace Metasound
diff --git a/Engine/Plugins/Runtime/Metasound/Source/MetasoundStandardNodes/Public/MetasoundStandardNodesCategories.h b/Engine/Plugins/Runtime/Metasound/Source/MetasoundStandardNodes/Public/MetasoundStandardNodesCategories.h
--- a/Engine/Plugins/Runtime/Metasound/Source/MetasoundStandardNodes/Public/MetasoundStandardNodesCategories.h
+++ b/Engine/Plugins/Runtime/Metasound/Source/MetasoundStandardNodes/Public/MetasoundStandardNodesCategories.h
@@ -11,6 +11,7 @@ namespace Metasound
 		extern const FText Generators;
 		extern const FText Filters;
 		extern const FText TriggerUtils;
+		extern const FText Music;
 		extern const FText RandomUtils;
 		extern const FText Conversions;
 		extern const FText DebugUtils;
diff --git a/Engine/Source/Runtime/SignalProcessing/Private/MidiNoteQuantizer.cpp b/Engine/Source/Runtime/SignalProcessing/Private/MidiNoteQuantizer.cpp
new file mode 100644
--- /dev/null
+++ b/Engine/Source/Runtime/SignalProcessing/Private/MidiNoteQuantizer.cpp
@@ -0,0 +1,212 @@
+// Copyright Epic Games, Inc. All Rights Reserved.
+
+#include "DSP/MidiNoteQuantizer.h"
+#include "Math/NumericLimits.h"
+
+namespace Audio
+{
+	// useful musical defines for scale degrees + accidentals to translate to half-steps
+	// i.e. - n1 = "Natural 1" = 0.0 half-steps from the root
+	//		- b3 = "Flat 3"    = 3.0f half-steps etc.
+	static const float n8 = 12.0f;  // C
+	static const float n7 = 11.0f;	// B
+	static const float		b7 = 10.0f;
+	static const float bb7 = 9.0f;
+	static const float n6 = 9.0f;	// A
+	static const float		b6 = 8.0f;
+	static const float		s5 = 8.0f;
+	static const float n5 = 7.0f;	// G
+	static const float		b5 = 6.0f;
+	static const float		s4 = 6.0f;
+	static const float n4 = 5.0f;	// F
+	static const float n3 = 4.0f;	// E
+	static const float		b3 = 3.0f;
+	static const float		s2 = 3.0f;
+	static const float n2 = 2.0f;	// D
+	static const float		b2 = 1.0f;
+	static const float n1 = 0.0f;	// C
+
+
+	ScaleDegreeSet::ScaleDegreeSet(const TArray<float>& InScaleDegrees, const TArray<float>& InChordTones)
+		: ScaleDegrees(InScaleDegrees)
+		, ChordTones(InChordTones)
+	{
+	}
+
+	TArrayView<float> ScaleDegreeSet::GetScaleDegreeSet(bool bChordTonesOnlyIfApplicable)
+	{
+		if (bChordTonesOnlyIfApplicable && ChordTones.Num())
+		{
+			return ChordTones;
+		}
+
+		return ScaleDegrees;
+	}
+
+	float FMidiNoteQuantizer::QuantizeMidiNote(const float InNote, const float InRoot, EMusicalScale::Scale InScale, bool bChordTonesOnlyIfApplicable)
+	{
+		return QuantizeMidiNote(InNote, InRoot, ScaleDegreeSetMap[InScale].GetScaleDegreeSet(bChordTonesOnlyIfApplicable));
+	}
+
+	float FMidiNoteQuantizer::QuantizeMidiNote(const float InNote, const float InRoot, const TArrayView<float> InScaleDegrees)
+	{
+		// QuantizeValueToScaleDegree() works within a single octave, so compensate for the octave and root
+		const float SetRange = InScaleDegrees[InScaleDegrees.Num() - 1]; // (SetRange will == 12.0 for a normal octave)
+		const float NoteRootDelta = (InNote - InRoot);
+		const float InNoteOctave = FMath::FloorToFloat(NoteRootDelta / SetRange);
+		const float ValueToQuant = NoteRootDelta - (SetRange * InNoteOctave);
+
+		const float QuantizedValue = QuantizeValueToScaleDegree(ValueToQuant, InScaleDegrees);
+
+		// reconstruct the quantized midi note given the InRoot and original octave of InNote
+		return (SetRange * InNoteOctave + QuantizedValue) + InRoot;
+	}
+
+	float FMidiNoteQuantizer::QuantizeValueToScaleDegree(const float InValue, const TArrayView<float> InScaleDegrees)
+	{
+		// the Set we are quantizing to should have at least 2 elements
+		if (!ensure(InScaleDegrees.Num()))
+		{
+			return InValue;
+		}
+		else if (InScaleDegrees.Num() == 1)
+		{
+			// if the set only contains a single value, that is technically the quantized value
+			return InScaleDegrees[0];
+		}
+		else if (!ensureMsgf((0.0f <= InValue) && (InValue <= InScaleDegrees[InScaleDegrees.Num() - 1])
+			, TEXT("Pitch Quantization: InValue (%f) is not in the range of a single octave (Needs to satisfy: 0.0f <= InValue <= %f)")
+			, InValue, InScaleDegrees[InScaleDegrees.Num() - 1]))
+		{
+			return InValue; // bail
+		}
+
+		bool bNegativeDelta = false;
+		float PositiveDelta = TNumericLimits<float>::Max();
+		float NegativeDelta = TNumericLimits<float>::Lowest();
+
+		// In the case where InValue == last element in the note set
+		// we want to initialize FirstNegativeIndex like this for our
+		// post-loop logic to return the correct value:
+		//	--> return Set[FirstNegativeIndex - 1];
+		// (in this case "NegativeDelta" will never get updated before the loop completes)
+		int32 FirstNegativeIndex = InScaleDegrees.Num();
+		for (int32 i = 0; i < InScaleDegrees.Num(); ++i)
+		{
+			// InValue == Set[i]
+			if (FMath::IsNearlyEqual(InValue, InScaleDegrees[i]))
+			{
+				return InScaleDegrees[i];
+			}
+
+			// We are measuring Delta as (InValue - InScaleDegrees[i]).
+			// Delta should be positive until we pass by InValue in the set.
+			// Once we hit a negative value, we know which two candidates we need
+			// to compare for the shortest distance, and can exit the loop
+			const float CurrDelta = InValue - InScaleDegrees[i];
+			if (CurrDelta >= 0.0f)
+			{
+				PositiveDelta = CurrDelta;
+				continue;
+			}
+			else
+			{
+				NegativeDelta = CurrDelta;
+				FirstNegativeIndex = i;
+				break;
+			}
+		}
+
+		if (FMath::Abs(NegativeDelta) < FMath::Abs(PositiveDelta))
+		{
+			return InScaleDegrees[FirstNegativeIndex];
+		}
+		else
+		{
+			return InScaleDegrees[FirstNegativeIndex - 1];
+		}
+
+		ensure(false); // something bad happened
+		return 0.0f;
+	}
+
+
+	// Initialize our note & chord arrays
+
+	/* REQUIREMENTS:
+	*	The (optional) second array defines Chord Tones, if empty, the first array will be used always
+	*		- this is ideally a subset of the first array (user-facing UI/UX will imply this), but does not technically need to be and is not enforced.
+	* 
+	*	The first array MUST have at least 1 entry or it will ensure.
+	* 
+	*	Both array MUST be in ascending order, or quantization behavior will be broken. This is not enforced at runtime and quantization and will not warn/ensure
+	* 
+	*	The Arrays must represent a full octave (inclusive): the lowest element being (0.0 and the highest being 12.0)
+	*/
+	TMap<EMusicalScale::Scale, ScaleDegreeSet> FMidiNoteQuantizer::ScaleDegreeSetMap{
+		/* entry example:
+		* 	{ EMusicalScale::Scale::[SCALE NAME],
+			{{ [(required) SCALE DEGREES] }, { [(optional) CHORD TONES] }}}
+		*/
+		  { EMusicalScale::Scale::Major,
+			{{ n1, n2, n3, n4, n5, n6, n7, n8 }, { n1, n2, n3, n5, n7, n8 }}}
+		, { EMusicalScale::Scale::Dominant7th_Mixolydian,
+			{{ n1, n2, n3, n4, n5, n6, b7, n8 }, { n1, n2, n3, n5, b7, n8 }}}
+		, { EMusicalScale::Scale::Minor_Dorian,
+			{{ n1, n2, b3, n4, n5, b7, n8 }, { n1, n2, b3, n5, b7, n8 }}}
+		, { EMusicalScale::Scale::HalfDiminished_Locrian,
+			{{ n1, b2, b3, n4, b5, b6, b7, n8 }, { n1, b3, b5, b7, n8 }}}
+		, { EMusicalScale::Scale::Diminished,
+			{{ n1, n2, b3, n4, b5, b6, n6, n7, n8 }, { n1, b3, n5, bb7, n8 }}}
+		, { EMusicalScale::Scale::MajorPentatonic,
+			{{ n1, n2, n3, n5, n6, n8 }, { n1, n3, n5, n7, n8 }}}
+		, { EMusicalScale::Scale::Lydian,
+			{{ n1, n2, n3, s4, n5, n6, n7, n8 }, { n1, n3, n5, n7, n8 }}}
+		, { EMusicalScale::Scale::Bebop_Major,
+			{{ n1, n2, n3, s4, n5, s5, n6, n7, n8 }, { n1, n3, n5, n7, n8 }}}
+		, { EMusicalScale::Scale::HarmonicMajor,
+			{{ n1, n2, n3, n4, n5, b6, n7, n8 }, { n1, n3, n5, n7, n8 }}}
+		, { EMusicalScale::Scale::LydianAugmented,
+			{{ n1, n2, n3, s4, s5, n6, n7, n8 }, { n1, n3, s5, n7, n8 }}}
+		, { EMusicalScale::Scale::Augmented,
+			{{ n1, s2, n3, n4, n5, b6, n7, n8 }, { n1, n3, n5, n7, n8 }}}
+		, { EMusicalScale::Scale::SixthModeOfHarmonicMinor,
+			{{ n1, s2, n3, s4, n5, n6, n7, n8 }, { n1, n3, n5, n7, n8 }}}
+		, { EMusicalScale::Scale::Diminished_BeginWithHalfStep,
+			{{ n1, b2, s2, n3, s4, n5, n6, b7, n8 }, { n1, n3, n5, b7, n8 }}}
+		, { EMusicalScale::Scale::Blues,
+			{{ n1, b3, n4, s4, n5, b7, n8 }, { n1, b3, n5, b7, n8 }}}
+		, { EMusicalScale::Scale::Bebop_Dominant,
+			{{ n1, n2, n3, n4, n5, n6, b7, n7, n8 }, { n1, n3, n5, n7, n8 }}}
+		, { EMusicalScale::Scale::Spanish_or_Jewish,
+			{{ n1, b2, n3, n4, n5, b6, b7, n8 }, { n1, n3, n5, b7, n8 }}}
+		, { EMusicalScale::Scale::LydianDominant,
+			{{ n1, n2, n3, s4, n5, n6, b7, n8 }, { n1, n3, n5, b7, n8 }}}
+		, { EMusicalScale::Scale::Hindu,
+			{{ n1, n2, n3, n4, n5, b6, b7, n8 }, { n1, n3, n5, b7, n8 }}}
+		, { EMusicalScale::Scale::WholeTone,
+			{{ n1, n2, n3, s4, s5, b7, n8 }, { n1, n3, n5, b7, n8 }}}
+		, { EMusicalScale::Scale::Chromatic,
+			{{ n1, b2, n2, b3, n3, n4, b5, n5, b6, n6, b7, n7, n8 }, { /* NA */ }}}
+		, { EMusicalScale::Scale::DiminishedWholeTone,
+			{{ n1, b2, s2, n3, s4, n5, n6, b7, n8 }, { n1, n3, n5, b7, n8 }}}
+		, { EMusicalScale::Scale::MinorPentatonic,
+			{{ n1, b3, n4, n5, b7, n8 }, { n1, b3, n5, b7, n8 }}}
+		, { EMusicalScale::Scale::Bebop_Minor,
+			{{ n1, n2, b3, n3, n4, n5, n6, b7, n8 }, { n1, b3, n5, b7, n8 }}}
+		, { EMusicalScale::Scale::MelodicMinor,
+			{{ n1, n2, b3, n4, n5, n6, n7, n8 }, { n1, b3, n5, n7, n8 }}}
+		, { EMusicalScale::Scale::Bebop_MinorNumber2,
+			{{ n1, n2, b3, n4, n5, s5, n6, n7, n8 }, { n1, b3, n5, n7, n8 }}}
+		, { EMusicalScale::Scale::HarmonicMinor,
+			{{ n1, n2, b3, n4, n5, b6, n7, n8 }, { n1, b3, n5, n7, n8 }}}
+		, { EMusicalScale::Scale::Diminished_BeginWithWholeStep,
+			{{ n1, n2, b3, n4, s4, s5, n6, n7, n8 }, { n1, b3, s5, n7, n8 }}}
+		, { EMusicalScale::Scale::Phrygian,
+			{{ n1, b2, b3, n4, n5, b6, b7, n8 }, { n1, b3, n5, b7, n8 }}}
+		, { EMusicalScale::Scale::NaturalMinor_Aeolian,
+			{{ n1, n2, b3, n4, n5, b6, b7, n8 }, { n1, b3, n5, b7, n8 }}}
+		, { EMusicalScale::Scale::HalfDiminished_LocrianNumber2,
+			{{ n1, b2, b3, n4, b5, b6, b7, n8 }, { n1, b3, b5, b7, n8 }}}
+	};
+} // namespace Audio
\ No newline at end of file
diff --git a/Engine/Source/Runtime/SignalProcessing/Public/DSP/MidiNoteQuantizer.h b/Engine/Source/Runtime/SignalProcessing/Public/DSP/MidiNoteQuantizer.h
new file mode 100644
--- /dev/null
+++ b/Engine/Source/Runtime/SignalProcessing/Public/DSP/MidiNoteQuantizer.h
@@ -0,0 +1,75 @@
+// Copyright Epic Games, Inc. All Rights Reserved.
+
+#pragma once
+
+namespace Audio
+{
+	namespace EMusicalScale
+	{
+		// Any additions to this Enum should also be added to MetasoundMidiNoteQuantizerNode.cpp
+		// in the "#pragma region Metasound Enums" area to show up as options in that Node
+		enum Scale
+		{
+			Major,
+			Dominant7th_Mixolydian,
+			Minor_Dorian,
+			HalfDiminished_Locrian,
+			Diminished,
+			MajorPentatonic,
+			Lydian,
+			Bebop_Major,
+			HarmonicMajor,
+			LydianAugmented,
+			Augmented,
+			SixthModeOfHarmonicMinor,
+			Diminished_BeginWithHalfStep,
+			Blues,
+			Bebop_Dominant,
+			Spanish_or_Jewish,
+			LydianDominant,
+			Hindu,
+			WholeTone,
+			Chromatic,
+			DiminishedWholeTone,
+			MinorPentatonic,
+			Bebop_Minor,
+			MelodicMinor,
+			Bebop_MinorNumber2,
+			HarmonicMinor,
+			Diminished_BeginWithWholeStep,
+			Phrygian,
+			NaturalMinor_Aeolian,
+			HalfDiminished_LocrianNumber2,
+
+			Count
+		};
+	}
+
+	struct SIGNALPROCESSING_API ScaleDegreeSet
+	{
+	public:
+		// ctor
+		ScaleDegreeSet(const TArray<float>& InScaleDegrees, const TArray<float>& InChordTones = {});
+
+		TArrayView<float> GetScaleDegreeSet(bool bChordTonesOnlyIfApplicable = false);
+
+	private:
+		TArray<float> ScaleDegrees;
+		TArray<float> ChordTones;
+	};
+
+	class SIGNALPROCESSING_API FMidiNoteQuantizer
+	{
+	public:
+		static float QuantizeMidiNote(const float InNote, const float InRoot, EMusicalScale::Scale InScale, bool bChordTonesOnlyIfApplicable = false);
+		static float QuantizeMidiNote(const float InNote, const float InRoot, const TArrayView<float> InScaleDegrees);
+
+	private:
+		static float QuantizeValueToScaleDegree(const float InValue, const TArrayView<float> InScaleDegrees);
+
+	public:
+		// Statically-defined scale/chord tone definitions.
+		static TMap<EMusicalScale::Scale, ScaleDegreeSet> ScaleDegreeSetMap;
+
+	}; // class FMidNoteQuantizer
+} // namespace Audio
\ No newline at end of file
