CIS static analysis fix (unreachable code):  Removed a stale & reachable "you shouldn't get here" assert at the bottom of FMidiNoteQuantizer::QuantizeValueToScaleDegree().

diff --git a/Engine/Source/Runtime/SignalProcessing/Private/MidiNoteQuantizer.cpp b/Engine/Source/Runtime/SignalProcessing/Private/MidiNoteQuantizer.cpp
--- a/Engine/Source/Runtime/SignalProcessing/Private/MidiNoteQuantizer.cpp
+++ b/Engine/Source/Runtime/SignalProcessing/Private/MidiNoteQuantizer.cpp
@@ -125,9 +125,6 @@ namespace Audio
 		{
 			return InScaleDegrees[FirstNegativeIndex - 1];
 		}
-
-		ensure(false); // something bad happened
-		return 0.0f;
 	}
 
 
