Non-unity fix: Added missing include to midi note quantizer
#okforgith public

diff --git a/Engine/Source/Runtime/SignalProcessing/Public/DSP/MidiNoteQuantizer.h b/Engine/Source/Runtime/SignalProcessing/Public/DSP/MidiNoteQuantizer.h
--- a/Engine/Source/Runtime/SignalProcessing/Public/DSP/MidiNoteQuantizer.h
+++ b/Engine/Source/Runtime/SignalProcessing/Public/DSP/MidiNoteQuantizer.h
@@ -2,6 +2,8 @@
 
 #pragma once
 
+#include "CoreMinimal.h"
+
 namespace Audio
 {
 	namespace EMusicalScale
