Fixing missed samples in linear interpolation resampler

diff --git a/Engine/Plugins/Runtime/Metasound/Source/MetasoundEngine/Private/MetasoundWave.cpp b/Engine/Plugins/Runtime/Metasound/Source/MetasoundEngine/Private/MetasoundWave.cpp
--- a/Engine/Plugins/Runtime/Metasound/Source/MetasoundEngine/Private/MetasoundWave.cpp
+++ b/Engine/Plugins/Runtime/Metasound/Source/MetasoundEngine/Private/MetasoundWave.cpp
@@ -86,14 +86,12 @@ namespace Audio
 
 	bool FSimpleDecoderWrapper::SetWave(const FSoundWaveProxyPtr& InWave, float InStartTimeSeconds, float InInitialPitchShiftSemitones)
 	{
-		Wave = InWave;
-
 		bool bSuccessful = true;
 
-		if (Wave.IsValid())
+		if (InWave.IsValid())
 		{
 			// validate data
-			if (!ensureAlwaysMsgf(Wave->IsStreaming(), TEXT("Metasounds does not support Force Inline (sound must be streaming)")))
+			if (!ensureAlwaysMsgf(InWave->IsStreaming(), TEXT("Metasounds does not support Force Inline (sound must be streaming)")))
 			{
 				Wave.Reset();
 				Input.Reset();
@@ -104,24 +102,57 @@ namespace Audio
 				return false;
 			}
 
-			// initialize input/output data
-			InputSampleRate = Wave->GetSampleRate();
-			FsOutToInRatio = (OutputSampleRate / InputSampleRate);
 
-			NumChannels = Wave->GetNumChannels();
-			DecodeBlockSizeInSamples = DecodeBlockSizeInFrames * NumChannels;
+			// Determine which values differ so that correct things can be updated.
+			const bool bIsDifferentWave = (InWave.Get() != Wave.Get());
+			const bool bUpdateNumChannels = (InWave->GetNumChannels() != NumChannels);
+			const bool bUpdateInputSampleRate = (InWave->GetSampleRate() != InputSampleRate);
 
-			// set Circular Buffer capacity
-			int32 Capacity = FMath::Max(1, static_cast<int32>(OutputBlockSizeInFrames * NumChannels * (1.0f + FsOutToInRatio * MaxPitchShiftRatio) * 2));
-			OutputCircularBuffer.Reserve(Capacity, true /* bRetainExistingSamples */);
+			const bool bUpdateAudioFormat = bUpdateNumChannels || bUpdateInputSampleRate || !bIsInitialized;
+			const bool bReinitDecoder = bIsDifferentWave || !bIsInitialized;
 
+			Wave = InWave;
 
-			// try to initialize decoders
-			bSuccessful = InitializeDecodersInternal(Wave, InStartTimeSeconds);
+			if (bReinitDecoder)
+			{
+				// try to initialize decoders
+				bSuccessful = InitializeDecodersInternal(Wave, InStartTimeSeconds);
+			}
+			else
+			{
+				// If the wave has not changed, only need to seek.
+				bSuccessful = SeekToTime(InStartTimeSeconds);
+			}
 			bDecoderIsDone = !bSuccessful;
 
-			Resampler.Init(Audio::EResamplingMethod::Linear, FsOutToInRatio, NumChannels);
-			PitchShifter.Reset(NumChannels, InInitialPitchShiftSemitones);
+			if (bUpdateAudioFormat)
+			{
+				// initialize input/output data
+				InputSampleRate = Wave->GetSampleRate();
+				FsOutToInRatio = (OutputSampleRate / InputSampleRate);
+
+				NumChannels = Wave->GetNumChannels();
+				DecodeBlockSizeInSamples = DecodeBlockSizeInFrames * NumChannels;
+
+				// set Circular Buffer capacity
+				int32 Capacity = FMath::Max(1, static_cast<int32>(OutputBlockSizeInFrames * NumChannels * (1.0f + FsOutToInRatio * MaxPitchShiftRatio) * 2));
+				OutputCircularBuffer.Reserve(Capacity, true /* bRetainExistingSamples */);
+
+				Resampler.Init(Audio::EResamplingMethod::Linear, FsOutToInRatio, NumChannels);
+				bIsInitialized = true;
+			}
+
+			if (bUpdateNumChannels)
+			{
+				
+				PitchShifter.Reset(Wave->GetNumChannels(), InInitialPitchShiftSemitones);
+				// Have to discard previous samples if the channel count changes. 
+				OutputCircularBuffer.SetNum(0);
+			}
+			else
+			{
+				PitchShifter.UpdatePitchShift(InInitialPitchShiftSemitones);
+			}
 		}
 
 		return bSuccessful;
@@ -141,9 +172,25 @@ namespace Audio
 		return false;
 	}
 	
+	bool FSimpleDecoderWrapper::CanGenerateAudio() const
+	{
+		// If there is a valid decoder, than this object can generate audio.
+		const bool bCanDecoderGenerateAudio = !bDecoderIsDone && Input.IsValid() && Output.IsValid() && Decoder.IsValid() && (NumChannels > 0);
+
+		// If there is audio remaining in the output circular buffer, then this can 
+		// generate some audio.
+		const bool bIsAudioInOutputBuffer = (0 != OutputCircularBuffer.Num());
+
+		return bCanDecoderGenerateAudio || bIsAudioInOutputBuffer;
+	}
 
 	uint32 FSimpleDecoderWrapper::GenerateAudio(float* OutputDest, int32 NumOutputFrames, int32& OutNumFramesConsumed, float PitchShiftInCents, bool bIsLooping)
 	{
+		if (0 == NumChannels)
+		{
+			return 0;
+		}
+
 		const uint32 NumOutputSamples = NumOutputFrames * NumChannels;
 
 		OutNumFramesConsumed = 0;
@@ -180,6 +227,7 @@ namespace Audio
 				int32 NumResamplerOutputFrames = 0;
 				int32 Error = Resampler.ProcessAudio(PreSrcBuffer.GetData(), NumFramesDecoded, bDecoderIsDone, PostSrcBuffer.GetData(), MaxNumResamplerOutputFramesPerBlock, NumResamplerOutputFrames);
 				ensure(Error == 0);
+
 				if (!PostSrcBuffer.Num() || !NumResamplerOutputFrames)
 				{
 					continue;
@@ -219,6 +267,11 @@ namespace Audio
 		return NumOutputSamples; // update once we are aware of partial decode on last buffer
 	}
 
+	uint32 FSimpleDecoderWrapper::GetNumChannels() const
+	{
+		return NumChannels;
+	}
+
 	bool FSimpleDecoderWrapper::InitializeDecodersInternal(const FSoundWaveProxyPtr& InWave, float InStartTimeSeconds)
 	{
 		if (!ensure(InWave.IsValid()))
diff --git a/Engine/Plugins/Runtime/Metasound/Source/MetasoundEngine/Private/MetasoundWavePlayerNode.cpp b/Engine/Plugins/Runtime/Metasound/Source/MetasoundEngine/Private/MetasoundWavePlayerNode.cpp
--- a/Engine/Plugins/Runtime/Metasound/Source/MetasoundEngine/Private/MetasoundWavePlayerNode.cpp
+++ b/Engine/Plugins/Runtime/Metasound/Source/MetasoundEngine/Private/MetasoundWavePlayerNode.cpp
@@ -426,12 +426,12 @@ namespace Metasound
 			int32 NumOutputFramesToGenerate = (EndFrame - StartFrame);
 			int32 NumFramesGenerated = 0;
 
-			const bool bCanDecodeWave = IsCurrentWaveValid();
-			if (bCanDecodeWave)
+			const bool bCanDecoderGenerateAudio = Decoder.CanGenerateAudio();
+			if (bCanDecoderGenerateAudio)
 			{
 				ensure(Decoder.CanGenerateAudio());
 
-				const int32 NumInputChannels = CurrentWaveAsset->GetNumChannels();
+				const int32 NumInputChannels = Decoder.GetNumChannels();
 				const bool bNeedsUpmix = (NumInputChannels == 1);
 				const bool bNeedsDeinterleave = !bNeedsUpmix;
 
@@ -561,6 +561,8 @@ namespace Metasound
 				}
 				
 				// handle decoder having completed during it's decode
+				const bool bCannotGenerateAudio = !Decoder.CanGenerateAudio();
+				const bool bDidNotGenerateEnough = NumFramesGenerated < NumOutputFramesToGenerate;
 				const bool bReachedEOF = !Decoder.CanGenerateAudio() || (NumFramesGenerated < NumOutputFramesToGenerate);
 				if (bReachedEOF && !bIsLooping)
 				{
diff --git a/Engine/Plugins/Runtime/Metasound/Source/MetasoundEngine/Public/MetasoundWave.h b/Engine/Plugins/Runtime/Metasound/Source/MetasoundEngine/Public/MetasoundWave.h
--- a/Engine/Plugins/Runtime/Metasound/Source/MetasoundEngine/Public/MetasoundWave.h
+++ b/Engine/Plugins/Runtime/Metasound/Source/MetasoundEngine/Public/MetasoundWave.h
@@ -90,14 +90,13 @@ namespace Audio
 		 */
 		bool SeekToTime(const float InSeconds);
 
-		bool CanGenerateAudio() const
-		{
-			return !bDecoderIsDone && Input.IsValid() && Output.IsValid() && Decoder.IsValid();
-		}
+		bool CanGenerateAudio() const;
 
 		// returns number of samples written.   
 		uint32 GenerateAudio(float* OutputDest, int32 NumOutputFrames, int32& OutNumFramesConsumed, float PitchShiftInCents = 0.0f, bool bIsLooping = false);
 
+		uint32 GetNumChannels() const;
+
 
 	private:
 		FSoundWaveProxyPtr Wave;
@@ -120,19 +119,20 @@ namespace Audio
 		Audio::TCircularAudioBuffer<float> OutputCircularBuffer;
 
 		// meta data:
-		float InputSampleRate{ 44100.f };
-		float OutputSampleRate{ 44100.f };
+		float InputSampleRate{ -1.f };
+		float OutputSampleRate{ -1.f };
 		float FsOutToInRatio{ 1.f };
 		float MaxPitchShiftCents{ 12.0f };
 		float MaxPitchShiftRatio{ 2.0f };
 
-		uint32 NumChannels{ 1 };
+		uint32 NumChannels{ 0 };
 		uint32 DecodeBlockSizeInFrames{ 64 };
 		uint32 DecodeBlockSizeInSamples{ 64 };
 		uint32 OutputBlockSizeInFrames { 512 };
 		
 		bool bDecoderIsDone{ true };
 		bool bDecoderHasLooped{ false };
+		bool bIsInitialized{ false };
 
 	}; // class FSimpleDecoderWrapper
 } // namespace Audio
diff --git a/Engine/Source/Runtime/SignalProcessing/Private/InterpolatedLinearPitchShifter.cpp b/Engine/Source/Runtime/SignalProcessing/Private/InterpolatedLinearPitchShifter.cpp
--- a/Engine/Source/Runtime/SignalProcessing/Private/InterpolatedLinearPitchShifter.cpp
+++ b/Engine/Source/Runtime/SignalProcessing/Private/InterpolatedLinearPitchShifter.cpp
@@ -4,101 +4,104 @@
 
 namespace Audio
 {
+	FLinearPitchShifter::FLinearPitchShifter(int32 InNumChannels)
+	{
+		Reset(InNumChannels, 0.f, 64);
+	}
+
 	void FLinearPitchShifter::Reset(int32 InNumChannels, float InInitialPitchShiftSemitones, int32 InInterpLengthFrames)
 	{
+		NumChannels = InNumChannels;
 		InterpLengthFrames = InInterpLengthFrames;
 		InterpFramesRemaining = 0;
 		PitchShiftRatio.SetValue(FMath::Pow(2.0f, InInitialPitchShiftSemitones / 12.0f), 0);
-		NumChannels = InNumChannels;
 
 		PreviousFrame.Reset();
-		PreviousFrame.AddZeroed(NumChannels);
+		if (NumChannels > 0)
+		{
+			PreviousFrame.AddZeroed(NumChannels);
+		}
 
-		bInterpolateBetweenBuffers = false;
+		CurrentIndex = 0.f;
 	}
 
 	int32 FLinearPitchShifter::ProcessAudio(const TArrayView<float> InputBuffer, Audio::TCircularAudioBuffer<float>& OutputBuffer)
 	{
 		const int32 NumInputFrames = InputBuffer.Num() / NumChannels;
 		int32 OutputFramesRendered = 0; // return value
+		const float* InputBufferData = InputBuffer.GetData();
+		float* PreviousFrameData = PreviousFrame.GetData();
+
+		if (0 == NumInputFrames)
+		{
+			return 0;
+		}
 
 		// NOTE: CurrentIndex is a float in the range (-1.0f, (float)(NumInputFrames - 1.0f))
 		// the fractional portion of CurrentIndex is used to interpolate between Floor(CurrentIndex) and Ceil(CurrentIndex)
 		// if CurrentIndex < -1.0f then we missed an interpolation that should have occured in the previous buffer
-		check(CurrentIndex > -1.0f);
-
+		check(CurrentIndex >= -1.f);
+		check(PreviousFrame.Num() == NumChannels);
 
 		// Handle interpolations between index -1.0f and 0.0f
 		// i.e. interpolating between our last buffer and the current buffer
-		if (bInterpolateBetweenBuffers)
+		while (CurrentIndex < 0.0f)
 		{
-			// Alpha is between (0.0 and 1.0]
-			// 0.0 -> Final frame of the previous Input buffer
-			// 1.0 -> 0th frame of the current Input buffer
-			float Alpha = CurrentIndex + 1.0f;
-
-			// Once Alpha reaches 1.x, we are on frame 0.x of the current buffer
-			while (Alpha < 1.0f)
+			const float Alpha = CurrentIndex + 1.f;
+			for (int32 Chan = 0; Chan < NumChannels; ++Chan)
 			{
-				for (int32 Chan = 0; Chan < NumChannels; ++Chan)
-				{
-					const float A = PreviousFrame[Chan];
-					const float B = InputBuffer[Chan];
-					OutputBuffer.Push(FMath::Lerp(A, B, Alpha));
-				}
-
-				Alpha += GetNextIndexDelta();
-				++OutputFramesRendered;
+				const float A = PreviousFrameData[Chan];
+				const float B = InputBufferData[Chan];
+				OutputBuffer.Push(FMath::Lerp(A, B, Alpha));
 			}
 
-			// Alpha == 1.0 means CurrentIndex == 0.0
-			CurrentIndex = Alpha - 1.0f;
-			bInterpolateBetweenBuffers = false;
+			CurrentIndex += GetNextIndexDelta();
+			++OutputFramesRendered;
 		}
 
+
 		// Early Exit: copy full input buffer if no work needs to be done
 		// (i.e., not interpolating and pitch shift ratio is 1.0f)
 		if (!InterpFramesRemaining && FMath::IsNearlyEqual(1.0f, PitchShiftRatio.GetTarget()))
 		{
-			OutputBuffer.Push(InputBuffer.GetData(), NumInputFrames * NumChannels);
+			OutputBuffer.Push(InputBufferData, NumInputFrames * NumChannels);
 			return (NumInputFrames + OutputFramesRendered);
 		}
 
 		// Normal case: linear interpolation across the input buffer
-		const int32 LastFrame = NumInputFrames - 1;
-		while (FMath::CeilToInt(CurrentIndex) < LastFrame)
-		{
-			const int32 i = NumChannels * (int32)FMath::Floor(CurrentIndex);
-			float const* A = InputBuffer.GetData() + i;
-			float const* B = A + NumChannels;
-			const float Alpha = FMath::Frac(CurrentIndex);
+		const int32 LastFrameIndex = NumInputFrames - 1;
+		int32 LowerFrameIndex = FMath::FloorToInt(CurrentIndex);
+		check(LowerFrameIndex >= 0);
+		int32 UpperFrameIndex = LowerFrameIndex + 1;
 
+		while (UpperFrameIndex < NumInputFrames)
+		{
+			int32 LowerSampleIndex = NumChannels * LowerFrameIndex;
+			int32 UpperSampleIndex = NumChannels * UpperFrameIndex;
+			float Alpha = CurrentIndex - (float)LowerFrameIndex;
 			for (int32 Chan = 0; Chan < NumChannels; ++Chan)
 			{
-				OutputBuffer.Push(FMath::Lerp(*A++, *B++, Alpha));
+				OutputBuffer.Push(FMath::Lerp(InputBufferData[LowerSampleIndex], InputBufferData[UpperSampleIndex], Alpha));
+				LowerSampleIndex++;
+				UpperSampleIndex++;
 			}
 
 			++OutputFramesRendered;
 			CurrentIndex += GetNextIndexDelta();
+			LowerFrameIndex = FMath::FloorToInt(CurrentIndex);
+			UpperFrameIndex = LowerFrameIndex + 1;
 		}
 
 		// wrap our fractional index by the buffer size
-		CurrentIndex -= (float)(LastFrame);
+		CurrentIndex -= (float)(NumInputFrames);
 
-		// If -1 < CurrentIndex < 0.f, we are going to need to interpolate between
-		// the final frame of this buffer and the 0th frame of the next buffer.
-		// So we cache the final frame and raise a flag to use on the next call
-		if (CurrentIndex < 0.0f)
+		// We may need to interpolate between the final frame of this buffer and 
+		// the 0th frame of the next buffer. Cache the final frame to use on the next call
+		const int32 FinalSampleIndex = LastFrameIndex * NumChannels;
+		
+		for (int32 Chan = 0; Chan < NumChannels; ++Chan)
 		{
-			const int32 FinalFrameIndex = (NumInputFrames - 1) * NumChannels;
-
-			PreviousFrame.Reset(NumChannels);
-			for (int32 Chan = 0; Chan < NumChannels; ++Chan)
-			{
-				PreviousFrame.Add(InputBuffer[FinalFrameIndex + Chan]);
-			}
-
-			bInterpolateBetweenBuffers = true;
+			PreviousFrameData[Chan] = InputBufferData[FinalSampleIndex + Chan];
 		}
 
 		return OutputFramesRendered;
@@ -126,4 +129,4 @@ namespace Audio
 		return PitchShiftRatio.GetTarget();
 	}
 
-} // namespace Audio
\ No newline at end of file
+} // namespace Audio
diff --git a/Engine/Source/Runtime/SignalProcessing/Public/DSP/Dsp.h b/Engine/Source/Runtime/SignalProcessing/Public/DSP/Dsp.h
--- a/Engine/Source/Runtime/SignalProcessing/Public/DSP/Dsp.h
+++ b/Engine/Source/Runtime/SignalProcessing/Public/DSP/Dsp.h
@@ -892,9 +892,9 @@ namespace Audio
 		}
 
 		// Same as Pop(), but does not increment the read counter.
-		int32 Peek(SampleType* OutBuffer, uint32 NumSamples)
+		int32 Peek(SampleType* OutBuffer, uint32 NumSamples) const
 		{
-			SampleType* SrcBuffer = InternalBuffer.GetData();
+			const SampleType* SrcBuffer = InternalBuffer.GetData();
 			const uint32 ReadIndex = ReadCounter.GetValue();
 			const uint32 WriteIndex = WriteCounter.GetValue();
 
@@ -912,7 +912,7 @@ namespace Audio
 
 		// Peeks a single element.
 		// returns false if the element is empty.
-		bool Peek(SampleType& OutElement)
+		bool Peek(SampleType& OutElement) const
 		{
 			if (Num() == 0)
 			{
@@ -980,7 +980,7 @@ namespace Audio
 		}
 
 		// Get number of samples that can be popped off of the buffer.
-		uint32 Num()
+		uint32 Num() const
 		{
 			const int32 ReadIndex = ReadCounter.GetValue();
 			const int32 WriteIndex = WriteCounter.GetValue();
@@ -996,13 +996,13 @@ namespace Audio
 		}
 
 		// Get the current capacity of the buffer
-		uint32 GetCapacity()
+		uint32 GetCapacity() const
 		{
 			return Capacity;
 		}
 
 		// Get number of samples that can be pushed onto the buffer before it is full.
-		uint32 Remainder()
+		uint32 Remainder() const
 		{
 			const uint32 ReadIndex = ReadCounter.GetValue();
 			const uint32 WriteIndex = WriteCounter.GetValue();
diff --git a/Engine/Source/Runtime/SignalProcessing/Public/DSP/InterpolatedLinearPitchShifter.h b/Engine/Source/Runtime/SignalProcessing/Public/DSP/InterpolatedLinearPitchShifter.h
--- a/Engine/Source/Runtime/SignalProcessing/Public/DSP/InterpolatedLinearPitchShifter.h
+++ b/Engine/Source/Runtime/SignalProcessing/Public/DSP/InterpolatedLinearPitchShifter.h
@@ -12,7 +12,7 @@ namespace Audio
 	{
 	public:
 		// ctor
-		FLinearPitchShifter() = default;
+		FLinearPitchShifter(int32 InNumChannels = 0);
 
 		void Reset(int32 InNumChannels, float InInitialPitchShiftSemitones = 0.0f, int32 InInterpLengthFrames = 100);
 
@@ -29,10 +29,9 @@ namespace Audio
 		TArray<float> PreviousFrame;
 		FParam PitchShiftRatio;
 		float CurrentIndex{ 0.0f };
-		int32 NumChannels{ 1 };
+		int32 NumChannels{ 0 };
 		int32 InterpLengthFrames{ 100 };
 		int32 InterpFramesRemaining{ 0 };
-		bool bInterpolateBetweenBuffers{ false };
 
 	}; // class FLinearPitchShifter
-} // namespace Audio
\ No newline at end of file
+} // namespace Audio
