adding include guard #rb none

diff --git a/Engine/Source/Runtime/SignalProcessing/Public/DSP/VectorLinearResampler.h b/Engine/Source/Runtime/SignalProcessing/Public/DSP/VectorLinearResampler.h
--- a/Engine/Source/Runtime/SignalProcessing/Public/DSP/VectorLinearResampler.h
+++ b/Engine/Source/Runtime/SignalProcessing/Public/DSP/VectorLinearResampler.h
@@ -1,5 +1,7 @@
 // Copyright Epic Games, Inc. All Rights Reserved.
 
+#pragma once
+
 #include "CoreMinimal.h"
 #include "DSP/Dsp.h"
 
