Deprecate more AlignedFloatBuffer callsites and use FAlignedFloatBuffer
#rnx

diff --git a/Engine/Plugins/Experimental/MotoSynth/Source/MotoSynth/Private/MotoSynthSourceAsset.cpp b/Engine/Plugins/Experimental/MotoSynth/Source/MotoSynth/Private/MotoSynthSourceAsset.cpp
--- a/Engine/Plugins/Experimental/MotoSynth/Source/MotoSynth/Private/MotoSynthSourceAsset.cpp
+++ b/Engine/Plugins/Experimental/MotoSynth/Source/MotoSynth/Private/MotoSynthSourceAsset.cpp
@@ -314,7 +314,7 @@ void UMotoSynthSource::UpdateSourceData()
 
 void UMotoSynthSource::FilterSourceDataForAnalysis()
 {
-	Audio::AlignedFloatBuffer ScratchBuffer;
+	Audio::FAlignedFloatBuffer ScratchBuffer;
 	ScratchBuffer.AddUninitialized(SourceDataPCM.Num());
 
 	float* ScratchDataBufferPtr = ScratchBuffer.GetData();
@@ -359,7 +359,7 @@ void UMotoSynthSource::DynamicsProcessForAnalysis()
 		DynamicsProcessor.SetAttackTime(DynamicsProcessorAttackTimeMsec);
 		DynamicsProcessor.SetReleaseTime(DynamicsProcessorReleaseTimeMsec);
 
-		Audio::AlignedFloatBuffer DynamicsProcessorScratchBuffer;
+		Audio::FAlignedFloatBuffer DynamicsProcessorScratchBuffer;
 		DynamicsProcessorScratchBuffer.AddUninitialized(AnalysisBuffer.Num());
 
 		DynamicsProcessor.ProcessAudio(AnalysisBuffer.GetData(), AnalysisBuffer.Num(), DynamicsProcessorScratchBuffer.GetData());
@@ -392,7 +392,7 @@ void UMotoSynthSource::BuildGrainTableByFFT()
 	// TODO
 }
 
-static void GetBufferViewFromAnalysisBuffer(const Audio::AlignedFloatBuffer& InAnalysisBuffer, int32 StartingBufferIndex, int32 BufferSize, TArrayView<const float>& OutBufferView)
+static void GetBufferViewFromAnalysisBuffer(const Audio::FAlignedFloatBuffer& InAnalysisBuffer, int32 StartingBufferIndex, int32 BufferSize, TArrayView<const float>& OutBufferView)
 {
 	if (BufferSize > 0)
 	{
diff --git a/Engine/Plugins/Experimental/MotoSynth/Source/MotoSynth/Public/MotoSynthSourceAsset.h b/Engine/Plugins/Experimental/MotoSynth/Source/MotoSynth/Public/MotoSynthSourceAsset.h
--- a/Engine/Plugins/Experimental/MotoSynth/Source/MotoSynth/Public/MotoSynthSourceAsset.h
+++ b/Engine/Plugins/Experimental/MotoSynth/Source/MotoSynth/Public/MotoSynthSourceAsset.h
@@ -228,7 +228,7 @@ protected:
 #endif // WITH_EDITOR
 
 #if WITH_EDITORONLY_DATA
-	Audio::AlignedFloatBuffer AnalysisBuffer;
+	Audio::FAlignedFloatBuffer AnalysisBuffer;
 #endif
 
 	// Data containing PCM audio of the imported source asset (filled out by the factory)
diff --git a/Engine/Plugins/FX/Niagara/Source/Niagara/Classes/NiagaraDataInterfaceAudioOscilloscope.h b/Engine/Plugins/FX/Niagara/Source/Niagara/Classes/NiagaraDataInterfaceAudioOscilloscope.h
--- a/Engine/Plugins/FX/Niagara/Source/Niagara/Classes/NiagaraDataInterfaceAudioOscilloscope.h
+++ b/Engine/Plugins/FX/Niagara/Source/Niagara/Classes/NiagaraDataInterfaceAudioOscilloscope.h
@@ -70,15 +70,15 @@ private:
 	float ScopeInMilliseconds;
 
 	// The buffer we downsample PoppedBuffer to based on the Resolution property.
-	Audio::AlignedFloatBuffer PopBuffer;
-	Audio::AlignedFloatBuffer DownsampledBuffer;
+	Audio::FAlignedFloatBuffer PopBuffer;
+	Audio::FAlignedFloatBuffer DownsampledBuffer;
 
 	// Handle for the SRV used by the generated HLSL.
 	FReadBuffer GPUDownsampledBuffer;
 	FThreadSafeCounter NumChannelsInDownsampledBuffer;
 	
 	// Buffer read by VectorVM worker threads. This vector is guaranteed to not be mutated during the VectorVM tasks.
-	Audio::AlignedFloatBuffer VectorVMReadBuffer;
+	Audio::FAlignedFloatBuffer VectorVMReadBuffer;
 
 	FDelegateHandle DeviceCreatedHandle;
 	FDelegateHandle DeviceDestroyedHandle;
diff --git a/Engine/Plugins/FX/Niagara/Source/Niagara/Classes/NiagaraDataInterfaceAudioSpectrum.h b/Engine/Plugins/FX/Niagara/Source/Niagara/Classes/NiagaraDataInterfaceAudioSpectrum.h
--- a/Engine/Plugins/FX/Niagara/Source/Niagara/Classes/NiagaraDataInterfaceAudioSpectrum.h
+++ b/Engine/Plugins/FX/Niagara/Source/Niagara/Classes/NiagaraDataInterfaceAudioSpectrum.h
@@ -102,21 +102,21 @@ private:
 	float SampleRate;
 	float FFTScale;
 
-	Audio::AlignedFloatBuffer PopBuffer;
+	Audio::FAlignedFloatBuffer PopBuffer;
 	TUniquePtr< FSlidingBuffer > SlidingBuffer;
 
-	TArray<Audio::AlignedFloatBuffer> ChannelSpectrumBuffers;
+	TArray<Audio::FAlignedFloatBuffer> ChannelSpectrumBuffers;
 
 	TUniquePtr<Audio::FContiguousSparse2DKernelTransform> CQTKernel;
 	TUniquePtr<Audio::IFFTAlgorithm> FFTAlgorithm;
 
-	Audio::AlignedFloatBuffer InterleavedBuffer;
-	Audio::AlignedFloatBuffer DeinterleavedBuffer;
-	Audio::AlignedFloatBuffer FFTInputBuffer;
-	Audio::AlignedFloatBuffer FFTOutputBuffer;
-	Audio::AlignedFloatBuffer PowerSpectrumBuffer;
-	Audio::AlignedFloatBuffer SpectrumBuffer;
-	Audio::AlignedFloatBuffer WindowBuffer;
+	Audio::FAlignedFloatBuffer InterleavedBuffer;
+	Audio::FAlignedFloatBuffer DeinterleavedBuffer;
+	Audio::FAlignedFloatBuffer FFTInputBuffer;
+	Audio::FAlignedFloatBuffer FFTOutputBuffer;
+	Audio::FAlignedFloatBuffer PowerSpectrumBuffer;
+	Audio::FAlignedFloatBuffer SpectrumBuffer;
+	Audio::FAlignedFloatBuffer WindowBuffer;
 
 	// Handle for the SRV used by the generated HLSL.
 	FReadBuffer GPUBuffer;
diff --git a/Engine/Plugins/FX/Niagara/Source/Niagara/Private/NiagaraDataInterfaceAudioSpectrum.cpp b/Engine/Plugins/FX/Niagara/Source/Niagara/Private/NiagaraDataInterfaceAudioSpectrum.cpp
--- a/Engine/Plugins/FX/Niagara/Source/Niagara/Private/NiagaraDataInterfaceAudioSpectrum.cpp
+++ b/Engine/Plugins/FX/Niagara/Source/Niagara/Private/NiagaraDataInterfaceAudioSpectrum.cpp
@@ -53,7 +53,7 @@ float FNiagaraDataInterfaceProxySpectrum::GetSpectrumValue(float InNormalizedPos
 		return 0.f;
 	}
 
-	const Audio::AlignedFloatBuffer& Buffer = ChannelSpectrumBuffers[InChannelIndex];
+	const Audio::FAlignedFloatBuffer& Buffer = ChannelSpectrumBuffers[InChannelIndex];
 	const int32 MaxIndex = Buffer.Num() - 1;
 
 	if (MaxIndex < 0)
@@ -214,7 +214,7 @@ void FNiagaraDataInterfaceProxySpectrum::UpdateSpectrum()
 	{
 		for (int32 ChannelIndex = 0; ChannelIndex < ChannelSpectrumBuffers.Num(); ChannelIndex++)
 		{
-			Audio::AlignedFloatBuffer& Buffer = ChannelSpectrumBuffers[ChannelIndex];
+			Audio::FAlignedFloatBuffer& Buffer = ChannelSpectrumBuffers[ChannelIndex];
 			if (Buffer.Num() > 0)
 			{
 				FMemory::Memset(Buffer.GetData(), 0, sizeof(float) * Buffer.Num());
@@ -248,7 +248,7 @@ void FNiagaraDataInterfaceProxySpectrum::UpdateSpectrum()
 	FSlidingWindow SlidingWindow(*SlidingBuffer, PopBuffer, InterleavedBuffer);
 
 	int32 NumWindows = 0;
-	for (Audio::AlignedFloatBuffer& InterleavedWindow : SlidingWindow)
+	for (Audio::FAlignedFloatBuffer& InterleavedWindow : SlidingWindow)
 	{
 		if (0 == NumWindows)
 		{
@@ -295,7 +295,7 @@ void FNiagaraDataInterfaceProxySpectrum::UpdateSpectrum()
 		// Apply scaling for each channel.
 		for (int32 ChannelIndex = 0; ChannelIndex < NumChannels; ChannelIndex++)
 		{
-			Audio::AlignedFloatBuffer& Buffer = ChannelSpectrumBuffers[ChannelIndex];
+			Audio::FAlignedFloatBuffer& Buffer = ChannelSpectrumBuffers[ChannelIndex];
 
 			Audio::ArrayMultiplyByConstantInPlace(Buffer, LinearScale);
 
@@ -340,7 +340,7 @@ void FNiagaraDataInterfaceProxySpectrum::ResizeSpectrumBuffer(int32 InNumChannel
 
 	for (int32 ChannelIndex = 0; ChannelIndex < InNumChannels; ChannelIndex++)
 	{
-		Audio::AlignedFloatBuffer& Buffer = ChannelSpectrumBuffers[ChannelIndex];
+		Audio::FAlignedFloatBuffer& Buffer = ChannelSpectrumBuffers[ChannelIndex];
 		Buffer.Reset();
 		if (InNumBands > 0)
 		{
diff --git a/Engine/Plugins/Online/OnlineSubsystemUtils/Source/OnlineSubsystemUtils/Private/VoiceEngineImpl.cpp b/Engine/Plugins/Online/OnlineSubsystemUtils/Source/OnlineSubsystemUtils/Private/VoiceEngineImpl.cpp
--- a/Engine/Plugins/Online/OnlineSubsystemUtils/Source/OnlineSubsystemUtils/Private/VoiceEngineImpl.cpp
+++ b/Engine/Plugins/Online/OnlineSubsystemUtils/Source/OnlineSubsystemUtils/Private/VoiceEngineImpl.cpp
@@ -5,6 +5,7 @@
 #include "VoiceModule.h"
 #include "Voice.h"
 
+#include "DSP/BufferVectorOperations.h"
 #include "Sound/SoundWaveProcedural.h"
 #include "OnlineSubsystemUtils.h"
 #include "GameFramework/GameSession.h"
@@ -1074,7 +1075,7 @@ void FVoiceEndpoint::PatchInOutput(Audio::FPatchOutputStrongPtr& InOutput)
 	OutputPatch = InOutput;
 }
 
-bool FVoiceEndpoint::OnProcessAudioStream(Audio::AlignedFloatBuffer& OutputBuffer)
+bool FVoiceEndpoint::OnProcessAudioStream(Audio::FAlignedFloatBuffer& OutputBuffer)
 {
 	FScopeLock ScopeLock(&OutputPatchCriticalSection);
 
diff --git a/Engine/Plugins/Online/OnlineSubsystemUtils/Source/OnlineSubsystemUtils/Public/VoiceEngineImpl.h b/Engine/Plugins/Online/OnlineSubsystemUtils/Source/OnlineSubsystemUtils/Public/VoiceEngineImpl.h
--- a/Engine/Plugins/Online/OnlineSubsystemUtils/Source/OnlineSubsystemUtils/Public/VoiceEngineImpl.h
+++ b/Engine/Plugins/Online/OnlineSubsystemUtils/Source/OnlineSubsystemUtils/Public/VoiceEngineImpl.h
@@ -119,7 +119,7 @@ public:
 
 private:
 	int32 NumChannelsComingIn;
-	Audio::AlignedFloatBuffer DownmixBuffer;
+	Audio::FAlignedFloatBuffer DownmixBuffer;
 
 	TUniquePtr<Audio::IAudioMixerPlatformInterface> PlatformEndpoint;
 
diff --git a/Engine/Plugins/Runtime/AudioSynesthesia/Source/AudioSynesthesiaCore/Private/AutoCorrelationPitchDetector.cpp b/Engine/Plugins/Runtime/AudioSynesthesia/Source/AudioSynesthesiaCore/Private/AutoCorrelationPitchDetector.cpp
--- a/Engine/Plugins/Runtime/AudioSynesthesia/Source/AudioSynesthesiaCore/Private/AutoCorrelationPitchDetector.cpp
+++ b/Engine/Plugins/Runtime/AudioSynesthesia/Source/AudioSynesthesiaCore/Private/AutoCorrelationPitchDetector.cpp
@@ -128,13 +128,13 @@ namespace Audio
 	{
 	}
 
-	void FAutoCorrelationPitchDetector::DetectPitches(const AlignedFloatBuffer& InMonoAudio, TArray<FPitchInfo>& OutPitches)
+	void FAutoCorrelationPitchDetector::DetectPitches(const FAlignedFloatBuffer& InMonoAudio, TArray<FPitchInfo>& OutPitches)
 	{
 		TAutoSlidingWindow<float, FAudioBufferAlignedAllocator> SlidingWindow(SlidingBuffer, InMonoAudio, WindowBuffer);
 
 		const float InverseSampleRate = 1.f / SampleRate;
 
-		for (AlignedFloatBuffer& AnalysisBuffer : SlidingWindow)
+		for (FAlignedFloatBuffer& AnalysisBuffer : SlidingWindow)
 		{
 			int32 SampleCounter = (WindowCounter * SlidingBuffer.GetNumHopSamples()) + SlidingBuffer.GetNumWindowSamples() / 2;
 			const float Timestamp = SampleCounter * InverseSampleRate;
diff --git a/Engine/Plugins/Runtime/AudioSynesthesia/Source/AudioSynesthesiaCore/Private/ConstantQNRTFactory.cpp b/Engine/Plugins/Runtime/AudioSynesthesia/Source/AudioSynesthesiaCore/Private/ConstantQNRTFactory.cpp
--- a/Engine/Plugins/Runtime/AudioSynesthesia/Source/AudioSynesthesiaCore/Private/ConstantQNRTFactory.cpp
+++ b/Engine/Plugins/Runtime/AudioSynesthesia/Source/AudioSynesthesiaCore/Private/ConstantQNRTFactory.cpp
@@ -162,7 +162,7 @@ namespace Audio
 
 		if (bDownmixToMono)
 		{
-			for (const AlignedFloatBuffer& Window : SlidingWindow)
+			for (const FAlignedFloatBuffer& Window : SlidingWindow)
 			{
 				FMemory::Memset(MonoBuffer.GetData(), 0, sizeof(float) * NumWindowFrames);
 
@@ -184,7 +184,7 @@ namespace Audio
 		else
 		{
 			// Loop through entire array of input samples.
-			for (const AlignedFloatBuffer& Window : SlidingWindow)
+			for (const FAlignedFloatBuffer& Window : SlidingWindow)
 			{
 				// Each channel is analyzed seperately.
 				for(const TDeinterleaveView<float>::TChannel<FAudioBufferAlignedAllocator> Channel : TAutoDeinterleaveView<float, FAudioBufferAlignedAllocator>(Window, ChannelBuffer, NumChannels))
@@ -200,7 +200,7 @@ namespace Audio
 	/** Called when analysis of audio asset is complete. */
 	void FConstantQNRTWorker::Finalize(IAnalyzerNRTResult* OutResult) 
 	{
-		AlignedFloatBuffer EmptyArray;
+		FAlignedFloatBuffer EmptyArray;
 		bool bDoFlush = true;
 
 		AnalyzeMultichannel(EmptyArray, OutResult, bDoFlush);
@@ -229,7 +229,7 @@ namespace Audio
 		SlidingBuffer.Reset();
 	}
 
-	void FConstantQNRTWorker::AnalyzeWindow(const AlignedFloatBuffer& InWindow, int32 InChannelIndex, FConstantQNRTResult& OutResult)
+	void FConstantQNRTWorker::AnalyzeWindow(const FAlignedFloatBuffer& InWindow, int32 InChannelIndex, FConstantQNRTResult& OutResult)
 	{
 		// Run CQT Analyzer
 		ConstantQAnalyzer->CalculateCQT(InWindow.GetData(), CQTSpectrum);
diff --git a/Engine/Plugins/Runtime/AudioSynesthesia/Source/AudioSynesthesiaCore/Private/FFTPeakPitchDetector.cpp b/Engine/Plugins/Runtime/AudioSynesthesia/Source/AudioSynesthesiaCore/Private/FFTPeakPitchDetector.cpp
--- a/Engine/Plugins/Runtime/AudioSynesthesia/Source/AudioSynesthesiaCore/Private/FFTPeakPitchDetector.cpp
+++ b/Engine/Plugins/Runtime/AudioSynesthesia/Source/AudioSynesthesiaCore/Private/FFTPeakPitchDetector.cpp
@@ -135,13 +135,13 @@ namespace Audio
 	{
 	}
 
-	void FFFTPeakPitchDetector::DetectPitches(const AlignedFloatBuffer& InMonoAudio, TArray<FPitchInfo>& OutPitches)
+	void FFFTPeakPitchDetector::DetectPitches(const FAlignedFloatBuffer& InMonoAudio, TArray<FPitchInfo>& OutPitches)
 	{
 		TAutoSlidingWindow<float, FAudioBufferAlignedAllocator> SlidingWindow(SlidingBuffer, InMonoAudio, WindowBuffer);
 
 		const float InverseMultiply = 1.f / SampleRate;
 
-		for (AlignedFloatBuffer& AnalysisBuffer : SlidingWindow)
+		for (FAlignedFloatBuffer& AnalysisBuffer : SlidingWindow)
 		{
 			int32 SampleCounter = (WindowCounter * SlidingBuffer.GetNumHopSamples()) + SlidingBuffer.GetNumWindowSamples() / 2;
 			float Timestamp = SampleCounter * InverseMultiply;
diff --git a/Engine/Plugins/Runtime/AudioSynesthesia/Source/AudioSynesthesiaCore/Private/MaxStrengthPitchTracker.cpp b/Engine/Plugins/Runtime/AudioSynesthesia/Source/AudioSynesthesiaCore/Private/MaxStrengthPitchTracker.cpp
--- a/Engine/Plugins/Runtime/AudioSynesthesia/Source/AudioSynesthesiaCore/Private/MaxStrengthPitchTracker.cpp
+++ b/Engine/Plugins/Runtime/AudioSynesthesia/Source/AudioSynesthesiaCore/Private/MaxStrengthPitchTracker.cpp
@@ -52,7 +52,7 @@ namespace Audio
 	{
 	}
 
-	void FMaxStrengthPitchTracker::TrackPitches(const AlignedFloatBuffer& InMonoAudio, TArray<FPitchTrackInfo>& OutPitchTracks)
+	void FMaxStrengthPitchTracker::TrackPitches(const FAlignedFloatBuffer& InMonoAudio, TArray<FPitchTrackInfo>& OutPitchTracks)
 	{
 		if (PitchDetector.IsValid())
 		{
diff --git a/Engine/Plugins/Runtime/AudioSynesthesia/Source/AudioSynesthesiaCore/Private/OnsetAnalyzer.cpp b/Engine/Plugins/Runtime/AudioSynesthesia/Source/AudioSynesthesiaCore/Private/OnsetAnalyzer.cpp
--- a/Engine/Plugins/Runtime/AudioSynesthesia/Source/AudioSynesthesiaCore/Private/OnsetAnalyzer.cpp
+++ b/Engine/Plugins/Runtime/AudioSynesthesia/Source/AudioSynesthesiaCore/Private/OnsetAnalyzer.cpp
@@ -101,7 +101,7 @@ namespace Audio
 
 		// Slide over audio 
 		FSlidingWindow SlidingWindow(SlidingBuffer, InSamples, WorkingBuffer, bDoFlush);
-		for (const AlignedFloatBuffer& WindowBuffer : SlidingWindow)
+		for (const FAlignedFloatBuffer& WindowBuffer : SlidingWindow)
 		{
 			// Calculate onset strength per audio buffer.
 			float OnsetStrength = GetNextOnsetStrength(WindowBuffer);
diff --git a/Engine/Plugins/Runtime/AudioSynesthesia/Source/AudioSynesthesiaCore/Private/OnsetNRTFactory.cpp b/Engine/Plugins/Runtime/AudioSynesthesia/Source/AudioSynesthesiaCore/Private/OnsetNRTFactory.cpp
--- a/Engine/Plugins/Runtime/AudioSynesthesia/Source/AudioSynesthesiaCore/Private/OnsetNRTFactory.cpp
+++ b/Engine/Plugins/Runtime/AudioSynesthesia/Source/AudioSynesthesiaCore/Private/OnsetNRTFactory.cpp
@@ -180,7 +180,7 @@ namespace Audio
 
 		if (Settings.bDownmixToMono)
 		{
-			for (const AlignedFloatBuffer& Window : SlidingWindow)
+			for (const FAlignedFloatBuffer& Window : SlidingWindow)
 			{
 				FMemory::Memset(MonoBuffer.GetData(), 0, sizeof(float) * NumWindowFrames);
 
@@ -200,7 +200,7 @@ namespace Audio
 		else
 		{
 			// Loop through entire array of input samples.
-			for (const AlignedFloatBuffer& Window : SlidingWindow)
+			for (const FAlignedFloatBuffer& Window : SlidingWindow)
 			{
 				// Each channel is analyzed seperately.
 				for(const TDeinterleaveView<float>::TChannel<FAudioBufferAlignedAllocator> Channel : TAutoDeinterleaveView<float, FAudioBufferAlignedAllocator>(Window, ChannelBuffer, NumChannels))
@@ -215,7 +215,7 @@ namespace Audio
 	void FOnsetNRTWorker::Finalize(IAnalyzerNRTResult* OutResult) 
 	{
 		// Finish analyzing audio
-		AlignedFloatBuffer EmptyArray;
+		FAlignedFloatBuffer EmptyArray;
 		bool bDoFlush = true;
 
 		AnalyzeMultichannel(EmptyArray, OutResult, bDoFlush);
@@ -248,7 +248,7 @@ namespace Audio
 		Reset();
 	}
 
-	void FOnsetNRTWorker::AnalyzeWindow(const AlignedFloatBuffer& InWindow, int32 InChannelIndex, FOnsetNRTResult& OutResult)
+	void FOnsetNRTWorker::AnalyzeWindow(const FAlignedFloatBuffer& InWindow, int32 InChannelIndex, FOnsetNRTResult& OutResult)
 	{
 		TArray<float> Strengths;
 		OnsetStrengthAnalyzers[InChannelIndex]->CalculateOnsetStrengths(InWindow, Strengths);
diff --git a/Engine/Plugins/Runtime/AudioSynesthesia/Source/AudioSynesthesiaCore/Public/AutoCorrelationPitchDetector.h b/Engine/Plugins/Runtime/AudioSynesthesia/Source/AudioSynesthesiaCore/Public/AutoCorrelationPitchDetector.h
--- a/Engine/Plugins/Runtime/AudioSynesthesia/Source/AudioSynesthesiaCore/Public/AutoCorrelationPitchDetector.h
+++ b/Engine/Plugins/Runtime/AudioSynesthesia/Source/AudioSynesthesiaCore/Public/AutoCorrelationPitchDetector.h
@@ -45,7 +45,7 @@ namespace Audio
 			 * @param InMonoAudio - A mono audio buffer of any length.
 			 * @param OutPitches - This array is filled with pitch observations.
 			 */
-			virtual void DetectPitches(const AlignedFloatBuffer& InMonoAudio, TArray<FPitchInfo>& OutPitches) override;
+			virtual void DetectPitches(const FAlignedFloatBuffer& InMonoAudio, TArray<FPitchInfo>& OutPitches) override;
 
 			/** Resets internal audio buffers. This pitch detector does not produce any more pitches on Finalize. */
 			virtual void Finalize(TArray<FPitchInfo>& OutPitches) override;
@@ -60,8 +60,8 @@ namespace Audio
 			int32 WindowCounter;
 
 			TSlidingBuffer<float> SlidingBuffer;
-			AlignedFloatBuffer WindowBuffer;
-			AlignedFloatBuffer AutoCorrBuffer;
+			FAlignedFloatBuffer WindowBuffer;
+			FAlignedFloatBuffer AutoCorrBuffer;
 			TUniquePtr<FBlockCorrelator> Correlator;
 
 			TArray<int32> PeakIndices;
diff --git a/Engine/Plugins/Runtime/AudioSynesthesia/Source/AudioSynesthesiaCore/Public/ConstantQAnalyzer.h b/Engine/Plugins/Runtime/AudioSynesthesia/Source/AudioSynesthesiaCore/Public/ConstantQAnalyzer.h
--- a/Engine/Plugins/Runtime/AudioSynesthesia/Source/AudioSynesthesiaCore/Public/ConstantQAnalyzer.h
+++ b/Engine/Plugins/Runtime/AudioSynesthesia/Source/AudioSynesthesiaCore/Public/ConstantQAnalyzer.h
@@ -83,10 +83,10 @@ namespace Audio
 		/** Scale factor to normalize between different FFT Sizes */
 		FWindow Window;
 
-		AlignedFloatBuffer WindowedSamples;
-		AlignedFloatBuffer FFTOutputRealData;
-		AlignedFloatBuffer FFTOutputImagData;
-		AlignedFloatBuffer SpectrumBuffer;
+		FAlignedFloatBuffer WindowedSamples;
+		FAlignedFloatBuffer FFTOutputRealData;
+		FAlignedFloatBuffer FFTOutputImagData;
+		FAlignedFloatBuffer SpectrumBuffer;
 
 		TUniquePtr<FContiguousSparse2DKernelTransform> CQTTransform;
 	};
diff --git a/Engine/Plugins/Runtime/AudioSynesthesia/Source/AudioSynesthesiaCore/Public/ConstantQNRTFactory.h b/Engine/Plugins/Runtime/AudioSynesthesia/Source/AudioSynesthesiaCore/Public/ConstantQNRTFactory.h
--- a/Engine/Plugins/Runtime/AudioSynesthesia/Source/AudioSynesthesiaCore/Public/ConstantQNRTFactory.h
+++ b/Engine/Plugins/Runtime/AudioSynesthesia/Source/AudioSynesthesiaCore/Public/ConstantQNRTFactory.h
@@ -153,7 +153,7 @@ namespace Audio
 		void AnalyzeMultichannel(TArrayView<const float> InAudio, IAnalyzerNRTResult* OutResult, bool bDoFlush);
 
 		/** Analyze a single window of audio from a single channel */
-		void AnalyzeWindow(const AlignedFloatBuffer& InWindow, int32 InChannelIndex, FConstantQNRTResult& OutResult);
+		void AnalyzeWindow(const FAlignedFloatBuffer& InWindow, int32 InChannelIndex, FConstantQNRTResult& OutResult);
 
 		int32 NumFrames;
 		int32 NumChannels;
@@ -166,9 +166,9 @@ namespace Audio
 
 		float MonoScaling;
 
-		AlignedFloatBuffer HopBuffer;
-		AlignedFloatBuffer ChannelBuffer;
-		AlignedFloatBuffer MonoBuffer;
+		FAlignedFloatBuffer HopBuffer;
+		FAlignedFloatBuffer ChannelBuffer;
+		FAlignedFloatBuffer MonoBuffer;
 
 		TArray<float> CQTSpectrum;
 
diff --git a/Engine/Plugins/Runtime/AudioSynesthesia/Source/AudioSynesthesiaCore/Public/FFTPeakPitchDetector.h b/Engine/Plugins/Runtime/AudioSynesthesia/Source/AudioSynesthesiaCore/Public/FFTPeakPitchDetector.h
--- a/Engine/Plugins/Runtime/AudioSynesthesia/Source/AudioSynesthesiaCore/Public/FFTPeakPitchDetector.h
+++ b/Engine/Plugins/Runtime/AudioSynesthesia/Source/AudioSynesthesiaCore/Public/FFTPeakPitchDetector.h
@@ -51,7 +51,7 @@ namespace Audio
 			 * @param InMonoAudio - A mono audio buffer of any length.
 			 * @param OutPitches - This array is filled with pitch observations.
 			 */
-			virtual void DetectPitches(const AlignedFloatBuffer& InMonoAudio, TArray<FPitchInfo>& OutPitches) override;
+			virtual void DetectPitches(const FAlignedFloatBuffer& InMonoAudio, TArray<FPitchInfo>& OutPitches) override;
 
 			/** Resets internal audio buffers. This pitch detector does not produce any more pitches on Finalize. */
 			virtual void Finalize(TArray<FPitchInfo>& OutPitches) override;
@@ -70,9 +70,9 @@ namespace Audio
 			TUniquePtr<FPeakPicker> PeakPicker;
 
 			TSlidingBuffer<float> SlidingBuffer;
-			AlignedFloatBuffer WindowBuffer;
-			AlignedFloatBuffer ZeroPaddedAnalysisBuffer;
-			AlignedFloatBuffer ComplexSpectrumBuffer;
-			AlignedFloatBuffer PowerSpectrumBuffer;
+			FAlignedFloatBuffer WindowBuffer;
+			FAlignedFloatBuffer ZeroPaddedAnalysisBuffer;
+			FAlignedFloatBuffer ComplexSpectrumBuffer;
+			FAlignedFloatBuffer PowerSpectrumBuffer;
 	};
 }
diff --git a/Engine/Plugins/Runtime/AudioSynesthesia/Source/AudioSynesthesiaCore/Public/LoudnessAnalyzer.h b/Engine/Plugins/Runtime/AudioSynesthesia/Source/AudioSynesthesiaCore/Public/LoudnessAnalyzer.h
--- a/Engine/Plugins/Runtime/AudioSynesthesia/Source/AudioSynesthesiaCore/Public/LoudnessAnalyzer.h
+++ b/Engine/Plugins/Runtime/AudioSynesthesia/Source/AudioSynesthesiaCore/Public/LoudnessAnalyzer.h
@@ -108,12 +108,12 @@ namespace Audio
 
             FWindow Window;
 
-            AlignedFloatBuffer WindowedSamples;
-            AlignedFloatBuffer FFTOutputRealData;
-            AlignedFloatBuffer FFTOutputImagData;
+            FAlignedFloatBuffer WindowedSamples;
+            FAlignedFloatBuffer FFTOutputRealData;
+            FAlignedFloatBuffer FFTOutputImagData;
 
-            AlignedFloatBuffer CurveWeights;
-            AlignedFloatBuffer PowerSpectrumBuffer;
+            FAlignedFloatBuffer CurveWeights;
+            FAlignedFloatBuffer PowerSpectrumBuffer;
     };
 
 
@@ -164,7 +164,7 @@ namespace Audio
 
         private:
 
-            AlignedFloatBuffer MonoBuffer;
+            FAlignedFloatBuffer MonoBuffer;
             TArray<float> ChannelWeights;
             FLoudnessAnalyzer Analyzer;
     };
diff --git a/Engine/Plugins/Runtime/AudioSynesthesia/Source/AudioSynesthesiaCore/Public/MaxStrengthPitchTracker.h b/Engine/Plugins/Runtime/AudioSynesthesia/Source/AudioSynesthesiaCore/Public/MaxStrengthPitchTracker.h
--- a/Engine/Plugins/Runtime/AudioSynesthesia/Source/AudioSynesthesiaCore/Public/MaxStrengthPitchTracker.h
+++ b/Engine/Plugins/Runtime/AudioSynesthesia/Source/AudioSynesthesiaCore/Public/MaxStrengthPitchTracker.h
@@ -28,7 +28,7 @@ namespace Audio
 			virtual ~FMaxStrengthPitchTracker();
 
 			/** Append pitch tracks to processed audio. */
-			virtual void TrackPitches(const AlignedFloatBuffer& InMonoAudio, TArray<FPitchTrackInfo>& OutPitchTracks) override;
+			virtual void TrackPitches(const FAlignedFloatBuffer& InMonoAudio, TArray<FPitchTrackInfo>& OutPitchTracks) override;
 
 			/** Call when all audio has been processed. Appends pitch tracks to OutPitchTracks. */
 			virtual void Finalize(TArray<FPitchTrackInfo>& OutPitchTracks) override;
diff --git a/Engine/Plugins/Runtime/AudioSynesthesia/Source/AudioSynesthesiaCore/Public/OnsetAnalyzer.h b/Engine/Plugins/Runtime/AudioSynesthesia/Source/AudioSynesthesiaCore/Public/OnsetAnalyzer.h
--- a/Engine/Plugins/Runtime/AudioSynesthesia/Source/AudioSynesthesiaCore/Public/OnsetAnalyzer.h
+++ b/Engine/Plugins/Runtime/AudioSynesthesia/Source/AudioSynesthesiaCore/Public/OnsetAnalyzer.h
@@ -73,11 +73,11 @@ namespace Audio
 			int32 LagSpectraIndex;
 
 			TSlidingBuffer<float> SlidingBuffer;
-			AlignedFloatBuffer WorkingBuffer;
-			AlignedFloatBuffer WindowedSamples;
-			AlignedFloatBuffer FFTOutputRealData;
-			AlignedFloatBuffer FFTOutputImagData;
-			AlignedFloatBuffer FFTSpectrumBuffer;
+			FAlignedFloatBuffer WorkingBuffer;
+			FAlignedFloatBuffer WindowedSamples;
+			FAlignedFloatBuffer FFTOutputRealData;
+			FAlignedFloatBuffer FFTOutputImagData;
+			FAlignedFloatBuffer FFTSpectrumBuffer;
 
 			// Array of arrays to hold lag spectra.
 			TArray< TArray<float> > PreviousSpectra;
diff --git a/Engine/Plugins/Runtime/AudioSynesthesia/Source/AudioSynesthesiaCore/Public/OnsetNRTFactory.h b/Engine/Plugins/Runtime/AudioSynesthesia/Source/AudioSynesthesiaCore/Public/OnsetNRTFactory.h
--- a/Engine/Plugins/Runtime/AudioSynesthesia/Source/AudioSynesthesiaCore/Public/OnsetNRTFactory.h
+++ b/Engine/Plugins/Runtime/AudioSynesthesia/Source/AudioSynesthesiaCore/Public/OnsetNRTFactory.h
@@ -152,7 +152,7 @@ namespace Audio
 		/** Analyze audio with multiple channels interleaved. */
 		void AnalyzeMultichannel(TArrayView<const float> InAudio, IAnalyzerNRTResult* OutResult, bool bDoFlush);
 
-		void AnalyzeWindow(const AlignedFloatBuffer& InWindow, int32 InChannelIndex, FOnsetNRTResult& OutResult);
+		void AnalyzeWindow(const FAlignedFloatBuffer& InWindow, int32 InChannelIndex, FOnsetNRTResult& OutResult);
 
 		void ExtractOnsetsFromOnsetStrengths(FOnsetNRTResult& OutResult);
 
@@ -167,9 +167,9 @@ namespace Audio
 		int32 NumWindowSamples;
 		float MonoScaling;
 
-		AlignedFloatBuffer HopBuffer;
-		AlignedFloatBuffer ChannelBuffer;
-		AlignedFloatBuffer MonoBuffer;
+		FAlignedFloatBuffer HopBuffer;
+		FAlignedFloatBuffer ChannelBuffer;
+		FAlignedFloatBuffer MonoBuffer;
 
 		TUniquePtr<TSlidingBuffer<float> > SlidingBuffer;
 
diff --git a/Engine/Plugins/Runtime/AudioSynesthesia/Source/AudioSynesthesiaCore/Public/PitchTracker.h b/Engine/Plugins/Runtime/AudioSynesthesia/Source/AudioSynesthesiaCore/Public/PitchTracker.h
--- a/Engine/Plugins/Runtime/AudioSynesthesia/Source/AudioSynesthesiaCore/Public/PitchTracker.h
+++ b/Engine/Plugins/Runtime/AudioSynesthesia/Source/AudioSynesthesiaCore/Public/PitchTracker.h
@@ -36,7 +36,7 @@ namespace Audio
 			 * @param InMonoAudio - The most recent input buffer from a stream of audio.
 			 * @param OutPitches - The pitches detected during this call to `DetectPitches` are appended to this array. 
 			 */
-			virtual void DetectPitches(const AlignedFloatBuffer& InMonoAudio, TArray<FPitchInfo>& OutPitches) = 0;
+			virtual void DetectPitches(const FAlignedFloatBuffer& InMonoAudio, TArray<FPitchInfo>& OutPitches) = 0;
 
 
 			/** Mark the end of an audio stream. 
@@ -71,7 +71,7 @@ namespace Audio
 			 * @param InMonoAudio - The most recent input buffer from a stream of audio.
 			 * @param OutPitchTracks - The pitch tracks detected during this call to `TrackPitches` are appended to this array. 
 			 */
-			virtual void TrackPitches(const AlignedFloatBuffer& InMonoAudio, TArray<FPitchTrackInfo>& OutPitchTracks) = 0;
+			virtual void TrackPitches(const FAlignedFloatBuffer& InMonoAudio, TArray<FPitchTrackInfo>& OutPitchTracks) = 0;
 
 			/** Mark the end of an audio stream. 
 			 *
diff --git a/Engine/Plugins/Runtime/Metasound/Source/MetasoundFrontend/Public/MetasoundAudioBuffer.h b/Engine/Plugins/Runtime/Metasound/Source/MetasoundFrontend/Public/MetasoundAudioBuffer.h
--- a/Engine/Plugins/Runtime/Metasound/Source/MetasoundFrontend/Public/MetasoundAudioBuffer.h
+++ b/Engine/Plugins/Runtime/Metasound/Source/MetasoundFrontend/Public/MetasoundAudioBuffer.h
@@ -11,19 +11,19 @@
 /** Define which determines whether to check that the size of the audio buffer has not changed since initialization */
 #define METASOUNDGRAPHCORE_CHECKAUDIONUM !(UE_BUILD_SHIPPING || UE_BUILD_TEST)
 
-namespace Metasound 
+namespace Metasound
 {
 	/**  FAudioBuffer
 	 *
-	 * FAudioBuffer is the default buffer for passing audio data between nodes.  It should not be resized. 
-	 * It is compatible with functions accepting const references to AlignedFloatBuffer (const AlignedFloatBuffer&) 
-	 * arguments via an implicit conversion operator which exposes the underlying AlignedFloatBuffer container.
+	 * FAudioBuffer is the default buffer for passing audio data between nodes.  It should not be resized.
+	 * It is compatible with functions accepting const references to FAlignedFloatBuffer (const FAlignedFloatBuffer&)
+	 * arguments via an implicit conversion operator which exposes the underlying FAlignedFloatBuffer container.
 	 *
 	 */
 	class METASOUNDFRONTEND_API FAudioBuffer
 	{
 		public:
-			/** Create an FAudioBuffer with a specific number of samples. 
+			/** Create an FAudioBuffer with a specific number of samples.
 			 *
 			 * @param InNumSamples - Number of samples in buffer.
 			 */
@@ -85,8 +85,8 @@ namespace Metasound
 				return Buffer.Num();
 			}
 
-			/** Implicit conversion to Audio::AlignedFloatBuffer */
-			FORCEINLINE operator const Audio::AlignedFloatBuffer& () const
+			/** Implicit conversion to Audio::FAlignedFloatBuffer */
+			FORCEINLINE operator const Audio::FAlignedFloatBuffer& () const
 			{
 #if METASOUNDGRAPHCORE_CHECKAUDIONUM
 				UE_CLOG(InitialNum != Buffer.Num(), LogMetaSound, Error, TEXT("MetaSound audio buffer size change detected.  Audio Buffers should not be resized."));
@@ -94,11 +94,11 @@ namespace Metasound
 return Buffer;
 			}
 
-			/** Implicit conversion to Audio::AlignedFloatBuffer 
+			/** Implicit conversion to Audio::FAlignedFloatBuffer 
 			 *
-			 * WARNING: if the buffer is resized, it will cause errors. 
+			 * WARNING: if the buffer is resized, it will cause errors.
 			 */
-			FORCEINLINE operator Audio::AlignedFloatBuffer& ()
+			FORCEINLINE operator Audio::FAlignedFloatBuffer& ()
 			{
 #if METASOUNDGRAPHCORE_CHECKAUDIONUM
 				UE_CLOG(InitialNum != Buffer.Num(), LogMetaSound, Error, TEXT("MetaSound audio buffer size change detected.  Audio Buffers should not be resized."));
@@ -116,7 +116,7 @@ return Buffer;
 
 		private:
 
-			Audio::AlignedFloatBuffer Buffer;
+			Audio::FAlignedFloatBuffer Buffer;
 
 #if METASOUNDGRAPHCORE_CHECKAUDIONUM
 			int32 InitialNum;
diff --git a/Engine/Plugins/Runtime/Metasound/Source/MetasoundFrontend/Public/MetasoundRouter.h b/Engine/Plugins/Runtime/Metasound/Source/MetasoundFrontend/Public/MetasoundRouter.h
--- a/Engine/Plugins/Runtime/Metasound/Source/MetasoundFrontend/Public/MetasoundRouter.h
+++ b/Engine/Plugins/Runtime/Metasound/Source/MetasoundFrontend/Public/MetasoundRouter.h
@@ -659,7 +659,7 @@ namespace Metasound
 
 				if (DelayLines[ChannelIndex].Num() > DelayTimeInFrames)
 				{
-					Audio::AlignedFloatBuffer& IntermediateBuffer = IntermediaryBuffers[ChannelIndex];
+					Audio::FAlignedFloatBuffer& IntermediateBuffer = IntermediaryBuffers[ChannelIndex];
 					DelayLines[ChannelIndex].Pop(IntermediateBuffer.GetData(), IntermediateBuffer.Num());
 					DataChannelInputs[ChannelIndex].PushAudio(IntermediateBuffer.GetData(), IntermediateBuffer.Num());
 				}
@@ -695,7 +695,7 @@ namespace Metasound
 				for (int32 Index = 0; Index < MaxChannels; Index++)
 				{
 					DelayLines.Emplace(NumSamplesToBuffer);
-					Audio::AlignedFloatBuffer BufferForChannel;
+					Audio::FAlignedFloatBuffer BufferForChannel;
 					BufferForChannel.AddZeroed(Params.OperatorSettings.GetNumFramesPerBlock());
 					IntermediaryBuffers.Add(MoveTemp(BufferForChannel));
 				}
@@ -725,7 +725,7 @@ namespace Metasound
 
 		TArray<Audio::FPatchInput> DataChannelInputs;
 
-		TArray<Audio::AlignedFloatBuffer> IntermediaryBuffers;
+		TArray<Audio::FAlignedFloatBuffer> IntermediaryBuffers;
 
 		FSenderInitParams Params;
 	};
diff --git a/Engine/Plugins/Runtime/Metasound/Source/MetasoundGenerator/Private/MetasoundGenerator.cpp b/Engine/Plugins/Runtime/Metasound/Source/MetasoundGenerator/Private/MetasoundGenerator.cpp
--- a/Engine/Plugins/Runtime/Metasound/Source/MetasoundGenerator/Private/MetasoundGenerator.cpp
+++ b/Engine/Plugins/Runtime/Metasound/Source/MetasoundGenerator/Private/MetasoundGenerator.cpp
@@ -8,7 +8,7 @@
 #include "MetasoundOperatorInterface.h"
 #include "MetasoundOutputNode.h"
 #include "MetasoundTrigger.h"
-//#include "MetasoundSource.h"
+
 
 namespace Metasound
 {
@@ -271,7 +271,7 @@ namespace Metasound
 		return bIsFinished;
 	}
 
-	int32 FMetasoundGenerator::FillWithBuffer(const Audio::AlignedFloatBuffer& InBuffer, float* OutAudio, int32 MaxNumOutputSamples)
+	int32 FMetasoundGenerator::FillWithBuffer(const Audio::FAlignedFloatBuffer& InBuffer, float* OutAudio, int32 MaxNumOutputSamples)
 	{
 		int32 InNum = InBuffer.Num();
 
diff --git a/Engine/Plugins/Runtime/Metasound/Source/MetasoundGenerator/Public/MetasoundGenerator.h b/Engine/Plugins/Runtime/Metasound/Source/MetasoundGenerator/Public/MetasoundGenerator.h
--- a/Engine/Plugins/Runtime/Metasound/Source/MetasoundGenerator/Public/MetasoundGenerator.h
+++ b/Engine/Plugins/Runtime/Metasound/Source/MetasoundGenerator/Public/MetasoundGenerator.h
@@ -151,7 +151,7 @@ namespace Metasound
 
 		// Fill OutAudio with data in InBuffer, up to maximum number of samples.
 		// Returns the number of samples used.
-		int32 FillWithBuffer(const Audio::AlignedFloatBuffer& InBuffer, float* OutAudio, int32 MaxNumOutputSamples); 
+		int32 FillWithBuffer(const Audio::FAlignedFloatBuffer& InBuffer, float* OutAudio, int32 MaxNumOutputSamples);
 
 		// Metasound creates deinterleaved audio while sound generator requires interleaved audio.
 		void InterleaveGeneratedAudio();
@@ -173,9 +173,9 @@ namespace Metasound
 		// Triggered when metasound is finished
 		FTriggerReadRef OnFinishedTriggerRef;
 
-		Audio::AlignedFloatBuffer InterleavedAudioBuffer;
+		Audio::FAlignedFloatBuffer InterleavedAudioBuffer;
 
-		Audio::AlignedFloatBuffer OverflowBuffer;
+		Audio::FAlignedFloatBuffer OverflowBuffer;
 
 		typedef FAsyncTask<FAsyncMetaSoundBuilder> FBuilderTask;
 		TUniquePtr<FBuilderTask> BuilderTask;
diff --git a/Engine/Plugins/Runtime/Metasound/Source/MetasoundGraphCore/Public/MetasoundOperatorSettings.h b/Engine/Plugins/Runtime/Metasound/Source/MetasoundGraphCore/Public/MetasoundOperatorSettings.h
--- a/Engine/Plugins/Runtime/Metasound/Source/MetasoundGraphCore/Public/MetasoundOperatorSettings.h
+++ b/Engine/Plugins/Runtime/Metasound/Source/MetasoundGraphCore/Public/MetasoundOperatorSettings.h
@@ -25,7 +25,7 @@ namespace Metasound
 	 * function.
 	 *
 	 * GetNumFramesPerBlock() is required to return a value that abides by the
-	 * alignment requirements of FAudioBuffer, AlignedFloatBuffer and the math
+	 * alignment requirements of FAudioBuffer, FAlignedFloatBuffer and the math
 	 * routines declared in BufferVectorOperations.h. This enables the vast 
 	 * majority of Metasound IOperators to take advantage of SIMD hardware
 	 * accelerations.  
@@ -74,10 +74,10 @@ namespace Metasound
 			/** Get the number of audio frames in a block. 
 			 *
 			 * This function is required to return a value that abides by the
-			 * alignment requirements of FAudioBuffer, AlignedFloatBuffer and 
+			 * alignment requirements of FAudioBuffer, FAlignedFloatBuffer and 
 			 * the math routines declared in BufferVectorOperations.h. This 
 			 * enables the vast majority of Metasound IOperators to take 
-			 * advantage of SIMD hardware accelerations.  
+			 * advantage of SIMD hardware accelerations.
 			 */
 			FORCEINLINE int32 GetNumFramesPerBlock() const
 			{
diff --git a/Engine/Plugins/Runtime/Oculus/OculusAudio/Source/OculusAudio/Private/OculusAmbisonicSpatializer.cpp b/Engine/Plugins/Runtime/Oculus/OculusAudio/Source/OculusAudio/Private/OculusAmbisonicSpatializer.cpp
--- a/Engine/Plugins/Runtime/Oculus/OculusAudio/Source/OculusAudio/Private/OculusAmbisonicSpatializer.cpp
+++ b/Engine/Plugins/Runtime/Oculus/OculusAudio/Source/OculusAudio/Private/OculusAmbisonicSpatializer.cpp
@@ -39,7 +39,7 @@ class FOculusSoundfieldBuffer : public ISoundfieldAudioPacket
 {
 public:
 	// Interleaved binaural audio buffer.
-	Audio::AlignedFloatBuffer AudioBuffer;
+	Audio::FAlignedFloatBuffer AudioBuffer;
 	int32 NumChannels;
 
 	FOculusSoundfieldBuffer()
@@ -74,8 +74,8 @@ private:
 	TArray<int32> SourceIds;
 
 	// temp buffers use to binauralize each source.
-	Audio::AlignedFloatBuffer ScratchMonoBuffer;
-	Audio::AlignedFloatBuffer ScratchOutputBuffer;
+	Audio::FAlignedFloatBuffer ScratchMonoBuffer;
+	Audio::FAlignedFloatBuffer ScratchOutputBuffer;
 
 	static int32 SourceIdCounter;
 	static FCriticalSection SourceIdCounterCritSection;
@@ -227,7 +227,7 @@ public:
 			else
 			{
 				// Otherwise, downmix and mix in.
-				Audio::AlignedFloatBuffer OutChannelMap;
+				Audio::FAlignedFloatBuffer OutChannelMap;
 				Audio::FMixerDevice::Get2DChannelMap(false, InputBuffer.NumChannels, InputData.PositionalData.NumChannels, false, OutChannelMap);
 				Audio::DownmixAndSumIntoBuffer(InputBuffer.NumChannels, InputData.PositionalData.NumChannels, InputBuffer.AudioBuffer, OutputData.AudioBuffer, OutChannelMap.GetData());
 			}
diff --git a/Engine/Plugins/Runtime/SoundFields/Source/SoundFields/Private/SoundFields.cpp b/Engine/Plugins/Runtime/SoundFields/Source/SoundFields/Private/SoundFields.cpp
--- a/Engine/Plugins/Runtime/SoundFields/Source/SoundFields/Private/SoundFields.cpp
+++ b/Engine/Plugins/Runtime/SoundFields/Source/SoundFields/Private/SoundFields.cpp
@@ -49,7 +49,7 @@ private:
 	FQuat PreviousRotation = FQuat::Identity;
 
 	FAmbisonicsSoundfieldBuffer RotatedAudio;
-	Audio::AlignedFloatBuffer DecoderOutputBuffer;
+	Audio::FAlignedFloatBuffer DecoderOutputBuffer;
 
 public:
 	virtual ~FAmbisonicsSoundfieldDecoder() {};
diff --git a/Engine/Plugins/Runtime/Synthesis/Source/Synthesis/Classes/SourceEffects/SourceEffectFilter.h b/Engine/Plugins/Runtime/Synthesis/Source/Synthesis/Classes/SourceEffects/SourceEffectFilter.h
--- a/Engine/Plugins/Runtime/Synthesis/Source/Synthesis/Classes/SourceEffects/SourceEffectFilter.h
+++ b/Engine/Plugins/Runtime/Synthesis/Source/Synthesis/Classes/SourceEffects/SourceEffectFilter.h
@@ -144,8 +144,8 @@ protected:
 	ESourceEffectFilterCircuit CircuitType;
 	ESourceEffectFilterType FilterType;
 
-	Audio::AlignedFloatBuffer ScratchModBuffer;
-	Audio::AlignedFloatBuffer ScratchEnvFollowerBuffer;
+	Audio::FAlignedFloatBuffer ScratchModBuffer;
+	Audio::FAlignedFloatBuffer ScratchEnvFollowerBuffer;
 
 	struct FAudioBusModulationData
 	{
diff --git a/Engine/Plugins/Runtime/Synthesis/Source/Synthesis/Classes/SourceEffects/SourceEffectMotionFilter.h b/Engine/Plugins/Runtime/Synthesis/Source/Synthesis/Classes/SourceEffects/SourceEffectMotionFilter.h
--- a/Engine/Plugins/Runtime/Synthesis/Source/Synthesis/Classes/SourceEffects/SourceEffectMotionFilter.h
+++ b/Engine/Plugins/Runtime/Synthesis/Source/Synthesis/Classes/SourceEffects/SourceEffectMotionFilter.h
@@ -300,8 +300,8 @@ protected:
 	TArray<float> ModDestinationUpdateTimeMS;
 
 	// Intermediary Scratch Buffers
-	Audio::AlignedFloatBuffer ScratchBufferA;
-	Audio::AlignedFloatBuffer ScratchBufferB;
+	Audio::FAlignedFloatBuffer ScratchBufferA;
+	Audio::FAlignedFloatBuffer ScratchBufferB;
 
 	// Filter Output Scalars
 	float FilterAMixScale;
diff --git a/Engine/Plugins/Runtime/Synthesis/Source/Synthesis/Classes/SubmixEffects/SubmixEffectStereoDelay.h b/Engine/Plugins/Runtime/Synthesis/Source/Synthesis/Classes/SubmixEffects/SubmixEffectStereoDelay.h
--- a/Engine/Plugins/Runtime/Synthesis/Source/Synthesis/Classes/SubmixEffects/SubmixEffectStereoDelay.h
+++ b/Engine/Plugins/Runtime/Synthesis/Source/Synthesis/Classes/SubmixEffects/SubmixEffectStereoDelay.h
@@ -72,7 +72,7 @@ public:
 
 private:
 
-	Audio::AlignedFloatBuffer ScratchStereoBuffer;
+	Audio::FAlignedFloatBuffer ScratchStereoBuffer;
 	Audio::FDelayStereo DelayStereo;
 };
 
diff --git a/Engine/Plugins/Runtime/Synthesis/Source/Synthesis/Classes/SynthComponents/SynthComponentToneGenerator.h b/Engine/Plugins/Runtime/Synthesis/Source/Synthesis/Classes/SynthComponents/SynthComponentToneGenerator.h
--- a/Engine/Plugins/Runtime/Synthesis/Source/Synthesis/Classes/SynthComponents/SynthComponentToneGenerator.h
+++ b/Engine/Plugins/Runtime/Synthesis/Source/Synthesis/Classes/SynthComponents/SynthComponentToneGenerator.h
@@ -26,7 +26,7 @@ public:
 private:
 	int32 NumChannels = 2;
 	Audio::FSineOsc SineOsc;
-	Audio::AlignedFloatBuffer Buffer;
+	Audio::FAlignedFloatBuffer Buffer;
 };
 
 UCLASS(ClassGroup = Synth, meta = (BlueprintSpawnableComponent))
diff --git a/Engine/Plugins/Runtime/Synthesis/Source/Synthesis/Private/ConvolutionReverb.cpp b/Engine/Plugins/Runtime/Synthesis/Source/Synthesis/Private/ConvolutionReverb.cpp
--- a/Engine/Plugins/Runtime/Synthesis/Source/Synthesis/Private/ConvolutionReverb.cpp
+++ b/Engine/Plugins/Runtime/Synthesis/Source/Synthesis/Private/ConvolutionReverb.cpp
@@ -124,10 +124,10 @@ namespace Audio
 			const int32 NumFrames = TargetImpulseSamples->Num() / InNumImpulseResponses;
 
 			// Prepare deinterleave pointers
-			TArray<AlignedFloatBuffer> DeinterleaveSamples;
+			TArray<FAlignedFloatBuffer> DeinterleaveSamples;
 			while (DeinterleaveSamples.Num() < InNumImpulseResponses)
 			{
-				AlignedFloatBuffer& Buffer = DeinterleaveSamples.Emplace_GetRef();
+				FAlignedFloatBuffer& Buffer = DeinterleaveSamples.Emplace_GetRef();
 				if (NumFrames > 0)
 				{
 					Buffer.AddUninitialized(NumFrames);
@@ -140,7 +140,7 @@ namespace Audio
 			// Set impulse responses in algorithm
 			for (int32 i = 0; i < DeinterleaveSamples.Num(); i++)
 			{
-				const AlignedFloatBuffer& Buffer = DeinterleaveSamples[i];
+				const FAlignedFloatBuffer& Buffer = DeinterleaveSamples[i];
 
 				InAlgo.SetImpulseResponse(i, Buffer.GetData(), Buffer.Num());
 			}
@@ -208,7 +208,7 @@ namespace Audio
 		return DefaultOutputFormat;
 	}
 	
-	void FConvolutionReverb::FChannelFormatConverterWrapper::ProcessAudio(const TArray<AlignedFloatBuffer>& InInputBuffers, TArray<AlignedFloatBuffer>& OutOutputBuffers)
+	void FConvolutionReverb::FChannelFormatConverterWrapper::ProcessAudio(const TArray<FAlignedFloatBuffer>& InInputBuffers, TArray<FAlignedFloatBuffer>& OutOutputBuffers)
 	{
 		if (nullptr != BaseConverter)
 		{
@@ -456,7 +456,7 @@ namespace Audio
 	}
 
 
-	void FConvolutionReverb::ProcessAudio(int32 NumInputChannels, AlignedFloatBuffer& InputAudio, int32 NumOutputChannels, AlignedFloatBuffer& OutputAudio)
+	void FConvolutionReverb::ProcessAudio(int32 NumInputChannels, FAlignedFloatBuffer& InputAudio, int32 NumOutputChannels, FAlignedFloatBuffer& OutputAudio)
 	{
 		check(InputBlockBuffer.IsValid());
 		check(OutputBlockBuffer.IsValid());
@@ -536,7 +536,7 @@ namespace Audio
 		}
 	}
 
-	void FConvolutionReverb::ProcessAudioBlock(const float* InputAudio, int32 InNumInputChannels, AlignedFloatBuffer& OutputAudio, int32 InNumOutputChannels)
+	void FConvolutionReverb::ProcessAudioBlock(const float* InputAudio, int32 InNumInputChannels, FAlignedFloatBuffer& OutputAudio, int32 InNumOutputChannels)
 	{
 		check(nullptr != InputAudio);
 		check(ConvolutionAlgorithm.IsValid());
@@ -697,7 +697,7 @@ namespace Audio
 		}
 	}
 
-	void FConvolutionReverb::ResizeArrayOfBuffers(TArray<AlignedFloatBuffer>& InArrayOfBuffers, int32 MinNumBuffers, int32 NumFrames) const
+	void FConvolutionReverb::ResizeArrayOfBuffers(TArray<FAlignedFloatBuffer>& InArrayOfBuffers, int32 MinNumBuffers, int32 NumFrames) const
 	{
 		while (InArrayOfBuffers.Num() < MinNumBuffers)
 		{
@@ -706,7 +706,7 @@ namespace Audio
 
 		for (int32 i = 0; i < MinNumBuffers; i++)
 		{
-			AlignedFloatBuffer& Buffer = InArrayOfBuffers[i];
+			FAlignedFloatBuffer& Buffer = InArrayOfBuffers[i];
 			Buffer.Reset();
 			if (NumFrames > 0)
 			{
@@ -754,7 +754,7 @@ namespace Audio
 		OutputBlockBuffer = MakeUnique<FAlignedBlockBuffer>(OutputCapacity, OutputMaxInspect);
 	}
 
-	void FConvolutionReverb::InterleaveBuffer(AlignedFloatBuffer& OutBuffer, const TArray<AlignedFloatBuffer>& InputBuffers, const int32 NumChannels)
+	void FConvolutionReverb::InterleaveBuffer(FAlignedFloatBuffer& OutBuffer, const TArray<FAlignedFloatBuffer>& InputBuffers, const int32 NumChannels)
 	{
 		check(InputBuffers.Num() >= NumChannels);
 
@@ -790,7 +790,7 @@ namespace Audio
 		}
 	}
 
-	void FConvolutionReverb::DeinterleaveBuffer(TArray<AlignedFloatBuffer>& OutputBuffers, TArrayView<const float> InputBuffer, const int32 NumChannels)
+	void FConvolutionReverb::DeinterleaveBuffer(TArray<FAlignedFloatBuffer>& OutputBuffers, TArrayView<const float> InputBuffer, const int32 NumChannels)
 	{
 		check(OutputBuffers.Num() >= NumChannels);
 
@@ -805,7 +805,7 @@ namespace Audio
 
 		for (int32 i = 0; i < NumChannels; i++)
 		{
-			AlignedFloatBuffer& OutBuffer = OutputBuffers[i];
+			FAlignedFloatBuffer& OutBuffer = OutputBuffers[i];
 
 			check(OutBuffer.Num() == NumFrames);
 
diff --git a/Engine/Plugins/Runtime/Synthesis/Source/Synthesis/Private/SubmixChannelFormatConverter.cpp b/Engine/Plugins/Runtime/Synthesis/Source/Synthesis/Private/SubmixChannelFormatConverter.cpp
--- a/Engine/Plugins/Runtime/Synthesis/Source/Synthesis/Private/SubmixChannelFormatConverter.cpp
+++ b/Engine/Plugins/Runtime/Synthesis/Source/Synthesis/Private/SubmixChannelFormatConverter.cpp
@@ -139,7 +139,7 @@ namespace Audio
 		static const bool bIsCenterChannelOnly = false;
 
 		// Get channel map from FMixerDevice
-		AlignedFloatBuffer ChannelMap;
+		FAlignedFloatBuffer ChannelMap;
 
 		FMixerDevice::Get2DChannelMap(bIsVorbis, InInputFormat.NumChannels, InOutputFormat.NumChannels, bIsCenterChannelOnly, ChannelMap);
 
@@ -198,7 +198,7 @@ namespace Audio
 		return OutputFormat;
 	}
 	
-	void FSimpleRouter::ProcessAudio(const TArray<AlignedFloatBuffer>& InInputBuffers, TArray<AlignedFloatBuffer>& OutOutputBuffers)
+	void FSimpleRouter::ProcessAudio(const TArray<FAlignedFloatBuffer>& InInputBuffers, TArray<FAlignedFloatBuffer>& OutOutputBuffers)
 	{
 		check(InInputBuffers.Num() == InputFormat.NumChannels);
 
diff --git a/Engine/Plugins/Runtime/Synthesis/Source/Synthesis/Private/SubmixChannelFormatConverter.h b/Engine/Plugins/Runtime/Synthesis/Source/Synthesis/Private/SubmixChannelFormatConverter.h
--- a/Engine/Plugins/Runtime/Synthesis/Source/Synthesis/Private/SubmixChannelFormatConverter.h
+++ b/Engine/Plugins/Runtime/Synthesis/Source/Synthesis/Private/SubmixChannelFormatConverter.h
@@ -72,7 +72,7 @@ namespace Audio
 			 * @param InInputBuffers - An array of input audio buffers.
 			 * @param OutOutputBuffers - An array of buffers where output audio is stored. 
 			 */
-			void ProcessAudio(const TArray<AlignedFloatBuffer>& InInputBuffers, TArray<AlignedFloatBuffer>& OutOutputBuffers) override;
+			void ProcessAudio(const TArray<FAlignedFloatBuffer>& InInputBuffers, TArray<FAlignedFloatBuffer>& OutOutputBuffers) override;
 
 			/** Create a simple router.
 			 *
diff --git a/Engine/Plugins/Runtime/Synthesis/Source/Synthesis/Private/SubmixEffectMultiBandCompressor.cpp b/Engine/Plugins/Runtime/Synthesis/Source/Synthesis/Private/SubmixEffectMultiBandCompressor.cpp
--- a/Engine/Plugins/Runtime/Synthesis/Source/Synthesis/Private/SubmixEffectMultiBandCompressor.cpp
+++ b/Engine/Plugins/Runtime/Synthesis/Source/Synthesis/Private/SubmixEffectMultiBandCompressor.cpp
@@ -169,8 +169,8 @@ void FSubmixEffectMultibandCompressor::OnProcessAudio(const FSoundEffectSubmixIn
 		return;
 	}
 
-	const Audio::AlignedFloatBuffer& InBuffer = *InData.AudioBuffer;
-	Audio::AlignedFloatBuffer& OutBuffer = *OutData.AudioBuffer;
+	const Audio::FAlignedFloatBuffer& InBuffer = *InData.AudioBuffer;
+	Audio::FAlignedFloatBuffer& OutBuffer = *OutData.AudioBuffer;
 
 	FMemory::Memzero(OutBuffer.GetData(), FrameSize * InData.NumFrames);
 
diff --git a/Engine/Plugins/Runtime/Synthesis/Source/Synthesis/Public/ConvolutionReverb.h b/Engine/Plugins/Runtime/Synthesis/Source/Synthesis/Public/ConvolutionReverb.h
--- a/Engine/Plugins/Runtime/Synthesis/Source/Synthesis/Public/ConvolutionReverb.h
+++ b/Engine/Plugins/Runtime/Synthesis/Source/Synthesis/Public/ConvolutionReverb.h
@@ -165,13 +165,13 @@ namespace Audio
 
 
 		// If the number of input frames changes between callbacks, the output may contain discontinuities.
-		void ProcessAudio(int32 InNumInputChannels, AlignedFloatBuffer& InputAudio, int32 InNumOutputChannels, AlignedFloatBuffer& OutputAudio);
+		void ProcessAudio(int32 InNumInputChannels, FAlignedFloatBuffer& InputAudio, int32 InNumOutputChannels, FAlignedFloatBuffer& OutputAudio);
 
 		int32 GetNumInputChannels() const;
 		int32 GetNumOutputChannels() const;
 
-		static void InterleaveBuffer(AlignedFloatBuffer& OutBuffer, const TArray<AlignedFloatBuffer>& InputBuffers, const int32 NumChannels);
-		static void DeinterleaveBuffer(TArray<AlignedFloatBuffer>& OutputBuffers, TArrayView<const float> InputBuffer, const int32 NumChannels);
+		static void InterleaveBuffer(FAlignedFloatBuffer& OutBuffer, const TArray<FAlignedFloatBuffer>& InputBuffers, const int32 NumChannels);
+		static void DeinterleaveBuffer(TArray<FAlignedFloatBuffer>& OutputBuffers, TArrayView<const float> InputBuffer, const int32 NumChannels);
 
 	private:
 
@@ -198,7 +198,7 @@ namespace Audio
 
 				const FOutputFormat& GetOutputFormat() const;
 				
-				void ProcessAudio(const TArray<AlignedFloatBuffer>& InInputBuffers, TArray<AlignedFloatBuffer>& OutOutputBuffers);
+				void ProcessAudio(const TArray<FAlignedFloatBuffer>& InInputBuffers, TArray<FAlignedFloatBuffer>& OutOutputBuffers);
 
 				void SetRearChannelBleed(float InGain, bool bFadeToGain=true);
 
@@ -221,7 +221,7 @@ namespace Audio
 
 		void SetConvolutionAlgorithm(TUniquePtr<IConvolutionAlgorithm> InAlgorithm);
 
-		void ProcessAudioBlock(const float* InputAudio, int32 InNumInputChannels, AlignedFloatBuffer& OutputAudio, int32 InNumOutputChannels);
+		void ProcessAudioBlock(const float* InputAudio, int32 InNumInputChannels, FAlignedFloatBuffer& OutputAudio, int32 InNumOutputChannels);
 
 		void Update(bool bFadeToParams);
 
@@ -229,7 +229,7 @@ namespace Audio
 
 		void ResizeProcessingBuffers();
 
-		void ResizeArrayOfBuffers(TArray<AlignedFloatBuffer>& InArrayOfBuffers, int32 MinNumBuffers, int32 NumFrames) const;
+		void ResizeArrayOfBuffers(TArray<FAlignedFloatBuffer>& InArrayOfBuffers, int32 MinNumBuffers, int32 NumFrames) const;
 
 		FConvolutionReverbSettings Settings;
 
@@ -239,19 +239,19 @@ namespace Audio
 		FChannelFormatConverterWrapper OutputChannelFormatConverter;
 		
 		// data is passed to the convolution algorithm as 2D arrays
-		TArray<AlignedFloatBuffer> InputDeinterleaveBuffers;
-		TArray<AlignedFloatBuffer> InputChannelConverterBuffers; 
+		TArray<FAlignedFloatBuffer> InputDeinterleaveBuffers;
+		TArray<FAlignedFloatBuffer> InputChannelConverterBuffers; 
 
 		// data is recieved from the convolution algorithm as 2D arrays
-		TArray<AlignedFloatBuffer> OutputDeinterleaveBuffers;
-		TArray<AlignedFloatBuffer> OutputChannelConverterBuffers;
+		TArray<FAlignedFloatBuffer> OutputDeinterleaveBuffers;
+		TArray<FAlignedFloatBuffer> OutputChannelConverterBuffers;
 
 		TArray<float*> InputBufferPtrs;
 		TArray<float*> OutputBufferPtrs;
 
 		TUniquePtr<FAlignedBlockBuffer> InputBlockBuffer;
 		TUniquePtr<FAlignedBlockBuffer> OutputBlockBuffer;
-		AlignedFloatBuffer InterleavedOutputBlock;
+		FAlignedFloatBuffer InterleavedOutputBlock;
 
 		int32 ExpectedNumFramesPerCallback;
 		int32 NumInputSamplesPerBlock;
diff --git a/Engine/Source/Runtime/AudioExtensions/Public/IAudioEndpoint.h b/Engine/Source/Runtime/AudioExtensions/Public/IAudioEndpoint.h
--- a/Engine/Source/Runtime/AudioExtensions/Public/IAudioEndpoint.h
+++ b/Engine/Source/Runtime/AudioExtensions/Public/IAudioEndpoint.h
@@ -146,7 +146,7 @@ private:
 	TUniquePtr<Audio::FMixerNullCallback> RenderCallback;
 	
 	// If we have a render callback, we pop audio from the PatchMixer to this buffer before calling OnAudioCallback.
-	Audio::AlignedFloatBuffer BufferForRenderCallback;
+	Audio::FAlignedFloatBuffer BufferForRenderCallback;
 
 	// Owns the current settings for this endpoint.
 	TUniquePtr<IAudioEndpointSettingsProxy> CurrentSettings;
diff --git a/Engine/Source/Runtime/AudioExtensions/Public/IAudioExtensionPlugin.h b/Engine/Source/Runtime/AudioExtensions/Public/IAudioExtensionPlugin.h
--- a/Engine/Source/Runtime/AudioExtensions/Public/IAudioExtensionPlugin.h
+++ b/Engine/Source/Runtime/AudioExtensions/Public/IAudioExtensionPlugin.h
@@ -146,7 +146,7 @@ struct FAudioPluginSourceInputData
 	uint64 AudioComponentId;
 
 	// The audio input buffer
-	Audio::AlignedFloatBuffer* AudioBuffer;
+	Audio::FAlignedFloatBuffer* AudioBuffer;
 
 	// Number of channels of the source audio buffer.
 	int32 NumChannels;
@@ -161,7 +161,7 @@ struct FAudioPluginSourceInputData
 struct FAudioPluginSourceOutputData
 {
 	// The audio output buffer
-	Audio::AlignedFloatBuffer AudioBuffer;
+	Audio::FAlignedFloatBuffer AudioBuffer;
 };
 
 /** This is a class which should be overridden to provide users with settings to use for individual sounds */
diff --git a/Engine/Source/Runtime/AudioExtensions/Public/ISoundfieldFormat.h b/Engine/Source/Runtime/AudioExtensions/Public/ISoundfieldFormat.h
--- a/Engine/Source/Runtime/AudioExtensions/Public/ISoundfieldFormat.h
+++ b/Engine/Source/Runtime/AudioExtensions/Public/ISoundfieldFormat.h
@@ -160,7 +160,7 @@ struct FSoundfieldEncoderInputData
 	/*
 	 * Input buffer of interleaved floats. Each channel of the interleaved AudioBuffer corresponds to a channel index in PositionalData.
 	 */
-	Audio::AlignedFloatBuffer& AudioBuffer;
+	Audio::FAlignedFloatBuffer& AudioBuffer;
 
 	// Number of channels of the source audio buffer.
 	int32 NumChannels;
@@ -195,7 +195,7 @@ struct FSoundfieldDecoderInputData
 
 struct FSoundfieldDecoderOutputData
 {
-	Audio::AlignedFloatBuffer& AudioBuffer;
+	Audio::FAlignedFloatBuffer& AudioBuffer;
 };
 
 class AUDIOEXTENSIONS_API ISoundfieldDecoderStream
diff --git a/Engine/Source/Runtime/AudioMixer/Classes/SubmixEffects/AudioMixerSubmixEffectReverb.h b/Engine/Source/Runtime/AudioMixer/Classes/SubmixEffects/AudioMixerSubmixEffectReverb.h
--- a/Engine/Source/Runtime/AudioMixer/Classes/SubmixEffects/AudioMixerSubmixEffectReverb.h
+++ b/Engine/Source/Runtime/AudioMixer/Classes/SubmixEffects/AudioMixerSubmixEffectReverb.h
@@ -159,7 +159,7 @@ private:
 	// Level of wet/dry signal on current buffer
 	Audio::FWetDry CurrentWetDry;
 
-	Audio::AlignedFloatBuffer WetInputBuffer;
+	Audio::FAlignedFloatBuffer WetInputBuffer;
 
 	// Curve which maps old reverb times to new decay value
 	FRichCurve DecayCurve;
diff --git a/Engine/Source/Runtime/AudioMixer/Private/AudioMixerBlueprintLibrary.cpp b/Engine/Source/Runtime/AudioMixer/Private/AudioMixerBlueprintLibrary.cpp
--- a/Engine/Source/Runtime/AudioMixer/Private/AudioMixerBlueprintLibrary.cpp
+++ b/Engine/Source/Runtime/AudioMixer/Private/AudioMixerBlueprintLibrary.cpp
@@ -243,7 +243,7 @@ USoundWave* UAudioMixerBlueprintLibrary::StopRecordingOutput(const UObject* Worl
 		float ChannelCount;
 
 		// call the thing here.
-		Audio::AlignedFloatBuffer& RecordedBuffer = MixerDevice->StopRecording(SubmixToRecord, ChannelCount, SampleRate);
+		Audio::FAlignedFloatBuffer& RecordedBuffer = MixerDevice->StopRecording(SubmixToRecord, ChannelCount, SampleRate);
 
 		if (RecordedBuffer.Num() == 0)
 		{
diff --git a/Engine/Source/Runtime/AudioMixer/Private/AudioMixerBus.cpp b/Engine/Source/Runtime/AudioMixer/Private/AudioMixerBus.cpp
--- a/Engine/Source/Runtime/AudioMixer/Private/AudioMixerBus.cpp
+++ b/Engine/Source/Runtime/AudioMixer/Private/AudioMixerBus.cpp
@@ -128,10 +128,10 @@ namespace Audio
 					// If source channels are different than bus channels, we need to up-mix or down-mix
 					if (NumSourceChannels != NumChannels)
 					{
-						Audio::AlignedFloatBuffer ChannelMap;
+						Audio::FAlignedFloatBuffer ChannelMap;
 						SourceManager->Get2DChannelMap(AudioBusSend.SourceId, NumChannels, ChannelMap);
 
-						Audio::AlignedFloatBuffer DownmixedBuffer;
+						Audio::FAlignedFloatBuffer DownmixedBuffer;
 						DownmixedBuffer.AddUninitialized(NumOutputFrames * NumChannels);
 						Audio::DownmixBuffer(NumSourceChannels, NumChannels, SourceBufferPtr, DownmixedBuffer.GetData(), NumOutputFrames, ChannelMap.GetData());
 						Audio::MixInBufferFast(DownmixedBuffer.GetData(), BusDataBufferPtr, DownmixedBuffer.Num(), AudioBusSend.SendLevel);
@@ -152,12 +152,12 @@ namespace Audio
 		PatchMixerSplitter.ProcessAudio();
 	}
 
-	void FMixerAudioBus::CopyCurrentBuffer(Audio::AlignedFloatBuffer& InChannelMap, int32 InNumOutputChannels, AlignedFloatBuffer& OutBuffer, int32 NumOutputFrames) const
+	void FMixerAudioBus::CopyCurrentBuffer(Audio::FAlignedFloatBuffer& InChannelMap, int32 InNumOutputChannels, FAlignedFloatBuffer& OutBuffer, int32 NumOutputFrames) const
 	{
 		const float* RESTRICT CurrentBuffer = GetCurrentBusBuffer();
 
 		check(NumChannels != InNumOutputChannels);
-		Audio::AlignedFloatBuffer DownmixedBuffer;
+		Audio::FAlignedFloatBuffer DownmixedBuffer;
 		DownmixedBuffer.AddUninitialized(NumOutputFrames * NumChannels);
 
 		Audio::DownmixBuffer(NumChannels, InNumOutputChannels, CurrentBuffer, DownmixedBuffer.GetData(), NumOutputFrames, InChannelMap.GetData());
@@ -165,7 +165,7 @@ namespace Audio
 		Audio::MixInBufferFast(DownmixedBuffer.GetData(), OutBuffer.GetData(), DownmixedBuffer.Num(), 1.0f);
 	}
 
-	void FMixerAudioBus::CopyCurrentBuffer(int32 InNumOutputChannels, AlignedFloatBuffer& OutBuffer, int32 NumOutputFrames) const
+	void FMixerAudioBus::CopyCurrentBuffer(int32 InNumOutputChannels, FAlignedFloatBuffer& OutBuffer, int32 NumOutputFrames) const
 	{
 		const float* RESTRICT CurrentBuffer = GetCurrentBusBuffer();
 
diff --git a/Engine/Source/Runtime/AudioMixer/Private/AudioMixerChannelMaps.cpp b/Engine/Source/Runtime/AudioMixer/Private/AudioMixerChannelMaps.cpp
--- a/Engine/Source/Runtime/AudioMixer/Private/AudioMixerChannelMaps.cpp
+++ b/Engine/Source/Runtime/AudioMixer/Private/AudioMixerChannelMaps.cpp
@@ -270,12 +270,12 @@ namespace Audio
 		return 0;	
 	}
 
-	void FMixerDevice::Get2DChannelMap(bool bIsVorbis, const int32 NumSourceChannels, const bool bIsCenterChannelOnly, Audio::AlignedFloatBuffer& OutChannelMap) const
+	void FMixerDevice::Get2DChannelMap(bool bIsVorbis, const int32 NumSourceChannels, const bool bIsCenterChannelOnly, Audio::FAlignedFloatBuffer& OutChannelMap) const
 	{
 		Get2DChannelMap(bIsVorbis, NumSourceChannels, PlatformInfo.NumChannels, bIsCenterChannelOnly, OutChannelMap);
 	}
 
-	void FMixerDevice::Get2DChannelMap(bool bIsVorbis, const int32 NumSourceChannels, const int32 NumOutputChannels, const bool bIsCenterChannelOnly, Audio::AlignedFloatBuffer& OutChannelMap)
+	void FMixerDevice::Get2DChannelMap(bool bIsVorbis, const int32 NumSourceChannels, const int32 NumOutputChannels, const bool bIsCenterChannelOnly, Audio::FAlignedFloatBuffer& OutChannelMap)
 	{
 		if (NumSourceChannels <= 0 ||
 			NumOutputChannels <= 0 ||
diff --git a/Engine/Source/Runtime/AudioMixer/Private/AudioMixerDevice.cpp b/Engine/Source/Runtime/AudioMixer/Private/AudioMixerDevice.cpp
--- a/Engine/Source/Runtime/AudioMixer/Private/AudioMixerDevice.cpp
+++ b/Engine/Source/Runtime/AudioMixer/Private/AudioMixerDevice.cpp
@@ -703,7 +703,7 @@ namespace Audio
 		bDebugOutputEnabled = true;
 	}
 
-	bool FMixerDevice::OnProcessAudioStream(AlignedFloatBuffer& Output)
+	bool FMixerDevice::OnProcessAudioStream(FAlignedFloatBuffer& Output)
 	{
 		LLM_SCOPE(ELLMTag::AudioMixer);
 
@@ -1788,7 +1788,7 @@ namespace Audio
 		return SourceManager->GetNumActiveSources();
 	}
 
-	void FMixerDevice::Get3DChannelMap(const int32 InSubmixNumChannels, const FWaveInstance* InWaveInstance, float EmitterAzimith, float NormalizedOmniRadius, Audio::AlignedFloatBuffer& OutChannelMap)
+	void FMixerDevice::Get3DChannelMap(const int32 InSubmixNumChannels, const FWaveInstance* InWaveInstance, float EmitterAzimith, float NormalizedOmniRadius, Audio::FAlignedFloatBuffer& OutChannelMap)
 	{
 		// If we're center-channel only, then no need for spatial calculations, but need to build a channel map
 		if (InWaveInstance->bCenterChannelOnly)
@@ -1967,7 +1967,7 @@ namespace Audio
 		}
 	}
 
-	Audio::AlignedFloatBuffer& FMixerDevice::StopRecording(USoundSubmix* InSubmix, float& OutNumChannels, float& OutSampleRate)
+	Audio::FAlignedFloatBuffer& FMixerDevice::StopRecording(USoundSubmix* InSubmix, float& OutNumChannels, float& OutSampleRate)
 	{
 		// if we can find the submix here, record that submix. Otherwise, just record the master submix.
 		FMixerSubmixPtr FoundSubmix = GetSubmixInstance(InSubmix).Pin();
@@ -2500,7 +2500,7 @@ namespace Audio
 		return bIsMain;
 	}
 
-	void FMixerDevice::WhiteNoiseTest(AlignedFloatBuffer& Output)
+	void FMixerDevice::WhiteNoiseTest(FAlignedFloatBuffer& Output)
 	{
 		const int32 NumFrames = OpenStreamParams.NumFrames;
 		const int32 NumChannels = PlatformInfo.NumChannels;
@@ -2517,7 +2517,7 @@ namespace Audio
 		}
 	}
 
-	void FMixerDevice::SineOscTest(AlignedFloatBuffer& Output)
+	void FMixerDevice::SineOscTest(FAlignedFloatBuffer& Output)
 	{
 		const int32 NumFrames = OpenStreamParams.NumFrames;
 		const int32 NumChannels = PlatformInfo.NumChannels;
diff --git a/Engine/Source/Runtime/AudioMixer/Private/AudioMixerSource.cpp b/Engine/Source/Runtime/AudioMixer/Private/AudioMixerSource.cpp
--- a/Engine/Source/Runtime/AudioMixer/Private/AudioMixerSource.cpp
+++ b/Engine/Source/Runtime/AudioMixer/Private/AudioMixerSource.cpp
@@ -1563,7 +1563,7 @@ namespace Audio
 		bPrevAllowedSpatializationSetting = IsSpatializationCVarEnabled();
 	}
 
-	bool FMixerSource::ComputeMonoChannelMap(Audio::AlignedFloatBuffer& OutChannelMap)
+	bool FMixerSource::ComputeMonoChannelMap(Audio::FAlignedFloatBuffer& OutChannelMap)
 	{
 		if (IsUsingObjectBasedSpatialization())
 		{
@@ -1595,7 +1595,7 @@ namespace Audio
 		return false;
 	}
 
-	bool FMixerSource::ComputeStereoChannelMap(Audio::AlignedFloatBuffer& OutChannelMap)
+	bool FMixerSource::ComputeStereoChannelMap(Audio::FAlignedFloatBuffer& OutChannelMap)
 	{
 		// Only recalculate positional data if the source has moved a significant amount:
 		if (WaveInstance->GetUseSpatialization() && (!FMath::IsNearlyEqual(WaveInstance->AbsoluteAzimuth, PreviousAzimuth, 0.01f) || MixerSourceVoice->NeedsSpeakerMap()))
@@ -1655,7 +1655,7 @@ namespace Audio
 		return false;
 	}
 
-	bool FMixerSource::ComputeChannelMap(const int32 NumSourceChannels, Audio::AlignedFloatBuffer& OutChannelMap)
+	bool FMixerSource::ComputeChannelMap(const int32 NumSourceChannels, Audio::FAlignedFloatBuffer& OutChannelMap)
 	{
 		if (NumSourceChannels == 1)
 		{
diff --git a/Engine/Source/Runtime/AudioMixer/Private/AudioMixerSource.h b/Engine/Source/Runtime/AudioMixer/Private/AudioMixerSource.h
--- a/Engine/Source/Runtime/AudioMixer/Private/AudioMixerSource.h
+++ b/Engine/Source/Runtime/AudioMixer/Private/AudioMixerSource.h
@@ -86,13 +86,13 @@ namespace Audio
 		void UpdateChannelMaps();
 
 		/** Computes the mono-channel map. */
-		bool ComputeMonoChannelMap(Audio::AlignedFloatBuffer& OutChannelMap);
+		bool ComputeMonoChannelMap(Audio::FAlignedFloatBuffer& OutChannelMap);
 
 		/** Computes the stereo-channel map. */
-		bool ComputeStereoChannelMap(Audio::AlignedFloatBuffer& OutChannelMap);
+		bool ComputeStereoChannelMap(Audio::FAlignedFloatBuffer& OutChannelMap);
 
 		/** Compute the channel map based on the number of output and source channels. */
-		bool ComputeChannelMap(const int32 NumSourceChannels, Audio::AlignedFloatBuffer& OutChannelMap);
+		bool ComputeChannelMap(const int32 NumSourceChannels, Audio::FAlignedFloatBuffer& OutChannelMap);
 
 		/** Whether or not we should create the source voice with the HRTF spatializer. */
 		bool UseObjectBasedSpatialization() const;
@@ -136,7 +136,7 @@ namespace Audio
 		};
 
 		// Mapping of channel map types to channel maps. Determined by what submixes this source sends its audio to.
-		Audio::AlignedFloatBuffer ChannelMap;
+		Audio::FAlignedFloatBuffer ChannelMap;
 		FRWLock ChannelMapLock;
 
 		float PreviousAzimuth;
diff --git a/Engine/Source/Runtime/AudioMixer/Private/AudioMixerSourceManager.cpp b/Engine/Source/Runtime/AudioMixer/Private/AudioMixerSourceManager.cpp
--- a/Engine/Source/Runtime/AudioMixer/Private/AudioMixerSourceManager.cpp
+++ b/Engine/Source/Runtime/AudioMixer/Private/AudioMixerSourceManager.cpp
@@ -1294,7 +1294,7 @@ namespace Audio
 		});
 	}
 
-	void FMixerSourceManager::SetChannelMap(const int32 SourceId, const uint32 NumInputChannels, const Audio::AlignedFloatBuffer& ChannelMap, const bool bInIs3D, const bool bInIsCenterChannelOnly)
+	void FMixerSourceManager::SetChannelMap(const int32 SourceId, const uint32 NumInputChannels, const Audio::FAlignedFloatBuffer& ChannelMap, const bool bInIs3D, const bool bInIsCenterChannelOnly)
 	{
 		AUDIO_MIXER_CHECK(SourceId < NumTotalSources);
 		AUDIO_MIXER_CHECK(GameThreadInfo.bIsBusy[SourceId]);
@@ -1820,7 +1820,7 @@ namespace Audio
 					int32 NumAudioBusChannels = AudioBusPtr->GetNumChannels();
 					if (NumAudioBusChannels != SourceInfo.NumInputChannels)
 					{
-						Audio::AlignedFloatBuffer ChannelMap;
+						Audio::FAlignedFloatBuffer ChannelMap;
 						MixerDevice->Get2DChannelMap(SourceInfo.bIsVorbis, AudioBusPtr->GetNumChannels(), SourceInfo.NumInputChannels, SourceInfo.bIsCenterChannelOnly, ChannelMap);
 						AudioBusPtr->CopyCurrentBuffer(ChannelMap, SourceInfo.NumInputChannels, SourceInfo.PreDistanceAttenuationBuffer, NumFramesPlayed);
 					}
@@ -2471,7 +2471,7 @@ namespace Audio
 		}
 	}
 
-	void FMixerSourceManager::MixOutputBuffers(const int32 SourceId, int32 InNumOutputChannels, const float InSendLevel, EMixerSourceSubmixSendStage InSubmixSendStage, AlignedFloatBuffer& OutWetBuffer) const
+	void FMixerSourceManager::MixOutputBuffers(const int32 SourceId, int32 InNumOutputChannels, const float InSendLevel, EMixerSourceSubmixSendStage InSubmixSendStage, FAlignedFloatBuffer& OutWetBuffer) const
 	{
 		if (InSendLevel > 0.0f)
 		{
@@ -2486,7 +2486,7 @@ namespace Audio
 		}
 	}
 
-	void FMixerSourceManager::Get2DChannelMap(const int32 SourceId, int32 InNumOutputChannels, Audio::AlignedFloatBuffer& OutChannelMap)
+	void FMixerSourceManager::Get2DChannelMap(const int32 SourceId, int32 InNumOutputChannels, Audio::FAlignedFloatBuffer& OutChannelMap)
 	{
 		AUDIO_MIXER_CHECK_AUDIO_PLAT_THREAD(MixerDevice);
 
diff --git a/Engine/Source/Runtime/AudioMixer/Private/AudioMixerSourceManager.h b/Engine/Source/Runtime/AudioMixer/Private/AudioMixerSourceManager.h
--- a/Engine/Source/Runtime/AudioMixer/Private/AudioMixerSourceManager.h
+++ b/Engine/Source/Runtime/AudioMixer/Private/AudioMixerSourceManager.h
@@ -48,7 +48,7 @@ namespace Audio
 	struct FMixerSourceVoiceBuffer
 	{
 		/** PCM float data. */
-		AlignedFloatBuffer AudioData;
+		FAlignedFloatBuffer AudioData;
 
 		/** How many times this buffer will loop. */
 		int32 LoopCount = 0;
@@ -198,7 +198,7 @@ namespace Audio
 		void SetVolume(const int32 SourceId, const float Volume);
 		void SetDistanceAttenuation(const int32 SourceId, const float DistanceAttenuation);
 		void SetSpatializationParams(const int32 SourceId, const FSpatializationParams& InParams);
-		void SetChannelMap(const int32 SourceId, const uint32 NumInputChannels, const Audio::AlignedFloatBuffer& InChannelMap, const bool bInIs3D, const bool bInIsCenterChannelOnly);
+		void SetChannelMap(const int32 SourceId, const uint32 NumInputChannels, const Audio::FAlignedFloatBuffer& InChannelMap, const bool bInIs3D, const bool bInIsCenterChannelOnly);
 		void SetLPFFrequency(const int32 SourceId, const float Frequency);
 		void SetHPFFrequency(const int32 SourceId, const float Frequency);
 
@@ -218,11 +218,11 @@ namespace Audio
 		bool NeedsSpeakerMap(const int32 SourceId) const;
 		void ComputeNextBlockOfSamples();
 		void ClearStoppingSounds();
-		void MixOutputBuffers(const int32 SourceId, int32 InNumOutputChannels, const float InSendLevel, EMixerSourceSubmixSendStage InSubmixSendStage, AlignedFloatBuffer& OutWetBuffer) const;
+		void MixOutputBuffers(const int32 SourceId, int32 InNumOutputChannels, const float InSendLevel, EMixerSourceSubmixSendStage InSubmixSendStage, FAlignedFloatBuffer& OutWetBuffer) const;
 
 		// Retrieves a channel map for the given source ID for the given output channels
 		// can be used even when a source is 3D if the source is doing any kind of bus sending or otherwise needs a channel map
-		void Get2DChannelMap(const int32 SourceId, int32 InNumOutputChannels, Audio::AlignedFloatBuffer& OutChannelMap);
+		void Get2DChannelMap(const int32 SourceId, int32 InNumOutputChannels, Audio::FAlignedFloatBuffer& OutChannelMap);
 
 		// Called by a soundfield submix to get encoded audio.
 		// If this source wasn't encoded (possibly because it is paused or finished playing),
@@ -350,10 +350,10 @@ namespace Audio
 			int32 CurrentAudioChunkNumFrames;
 
 			// The post-attenuation source buffer, used to send audio to submixes
-			Audio::AlignedFloatBuffer SourceBuffer;
-			Audio::AlignedFloatBuffer PreEffectBuffer;
-			Audio::AlignedFloatBuffer PreDistanceAttenuationBuffer;
-			Audio::AlignedFloatBuffer SourceEffectScratchBuffer;
+			Audio::FAlignedFloatBuffer SourceBuffer;
+			Audio::FAlignedFloatBuffer PreEffectBuffer;
+			Audio::FAlignedFloatBuffer PreDistanceAttenuationBuffer;
+			Audio::FAlignedFloatBuffer SourceEffectScratchBuffer;
 
 			// Data used for delaying the rendering of source audio for sample-accurate quantization
 			int32 SubCallbackDelayLengthInFrames{ 0 };
@@ -425,7 +425,7 @@ namespace Audio
 			float HighpassModulationBase;
 
 			FSpatializationParams SpatParams;
-			Audio::AlignedFloatBuffer ScratchChannelMap;
+			Audio::FAlignedFloatBuffer ScratchChannelMap;
 
 			// Quantization data
 			FQuartzQuantizedCommandHandle QuantizedCommandHandle;
diff --git a/Engine/Source/Runtime/AudioMixer/Private/AudioMixerSourceOutputBuffer.cpp b/Engine/Source/Runtime/AudioMixer/Private/AudioMixerSourceOutputBuffer.cpp
--- a/Engine/Source/Runtime/AudioMixer/Private/AudioMixerSourceOutputBuffer.cpp
+++ b/Engine/Source/Runtime/AudioMixer/Private/AudioMixerSourceOutputBuffer.cpp
@@ -137,7 +137,7 @@ namespace Audio
 		PostAttenuationOutputBuffer.AddUninitialized(NumFrames * NumOutputChannels);
 	}
 
-	bool FMixerSourceSubmixOutputBuffer::SetChannelMap(const AlignedFloatBuffer& InChannelMap, bool bInIsCenterChannelOnly)
+	bool FMixerSourceSubmixOutputBuffer::SetChannelMap(const FAlignedFloatBuffer& InChannelMap, bool bInIsCenterChannelOnly)
 	{
 		bool bNeedsNewChannelMap = false;
 
@@ -146,7 +146,7 @@ namespace Audio
 
 		if (InChannelMap.Num() != ChannelMapSize)
 		{
-			AlignedFloatBuffer NewChannelMap;
+			FAlignedFloatBuffer NewChannelMap;
 
 			if (bIs3D)
 			{
@@ -169,12 +169,12 @@ namespace Audio
 		return bNeedsNewChannelMap;
 	}
 
-	void FMixerSourceSubmixOutputBuffer::SetPreAttenuationSourceBuffer(AlignedFloatBuffer* InPreAttenuationSourceBuffer)
+	void FMixerSourceSubmixOutputBuffer::SetPreAttenuationSourceBuffer(FAlignedFloatBuffer* InPreAttenuationSourceBuffer)
 	{
 		PreAttenuationSourceBuffer = InPreAttenuationSourceBuffer;
 	}
 
-	void FMixerSourceSubmixOutputBuffer::SetPostAttenuationSourceBuffer(AlignedFloatBuffer* InPostAttenuationSourceBuffer)
+	void FMixerSourceSubmixOutputBuffer::SetPostAttenuationSourceBuffer(FAlignedFloatBuffer* InPostAttenuationSourceBuffer)
 	{
 		PostAttenuationSourceBuffer = InPostAttenuationSourceBuffer;
 	}
@@ -208,7 +208,7 @@ namespace Audio
 		}
 	}
 
-	void FMixerSourceSubmixOutputBuffer::ComputeOutput3D(AlignedFloatBuffer& InSourceBuffer, AlignedFloatBuffer& OutSourceBuffer)
+	void FMixerSourceSubmixOutputBuffer::ComputeOutput3D(FAlignedFloatBuffer& InSourceBuffer, FAlignedFloatBuffer& OutSourceBuffer)
 	{
 		if (SoundfieldDecoder.IsValid())
 		{
@@ -296,7 +296,7 @@ namespace Audio
 		SourceChannelMap.CopyDestinationToStart();
 	}
 
-	void FMixerSourceSubmixOutputBuffer::ComputeOutput2D(AlignedFloatBuffer& InSourceBuffer, AlignedFloatBuffer& OutSourceBuffer)
+	void FMixerSourceSubmixOutputBuffer::ComputeOutput2D(FAlignedFloatBuffer& InSourceBuffer, FAlignedFloatBuffer& OutSourceBuffer)
 	{
 		// For 2D sources, we just apply the gain matrix in ChannelDestionationGains with no interpolation.
 		if (SoundfieldDecoder.IsValid())
@@ -444,7 +444,7 @@ namespace Audio
 		}
 	}
 
-	void FMixerSourceSubmixOutputBuffer::EncodeSoundfield(FSoundfieldData& InSoundfieldData, Audio::AlignedFloatBuffer& InSourceBuffer)
+	void FMixerSourceSubmixOutputBuffer::EncodeSoundfield(FSoundfieldData& InSoundfieldData, Audio::FAlignedFloatBuffer& InSourceBuffer)
 	{
 		check(InSoundfieldData.EncoderSettings.IsValid());
 		check(InSoundfieldData.EncodedPacket.IsValid());
@@ -487,7 +487,7 @@ namespace Audio
 		}
 	}
 
-	void FMixerSourceSubmixOutputBuffer::MixOutput(float InSendLevel, EMixerSourceSubmixSendStage InSubmixSendStage, AlignedFloatBuffer& OutMixedBuffer) const
+	void FMixerSourceSubmixOutputBuffer::MixOutput(float InSendLevel, EMixerSourceSubmixSendStage InSubmixSendStage, FAlignedFloatBuffer& OutMixedBuffer) const
 	{
 		if (InSubmixSendStage == EMixerSourceSubmixSendStage::PostDistanceAttenuation)
 		{
@@ -504,7 +504,7 @@ namespace Audio
 		return SoundfieldPositionalData.Rotation;
 	}
 
-	void FMixerSourceSubmixOutputBuffer::CopyReverbPluginOutputData(AlignedFloatBuffer& InAudioBuffer)
+	void FMixerSourceSubmixOutputBuffer::CopyReverbPluginOutputData(FAlignedFloatBuffer& InAudioBuffer)
 	{
 		ReverbPluginOutputBuffer.Reset();
 		ReverbPluginOutputBuffer.Append(InAudioBuffer);
diff --git a/Engine/Source/Runtime/AudioMixer/Private/AudioMixerSourceOutputBuffer.h b/Engine/Source/Runtime/AudioMixer/Private/AudioMixerSourceOutputBuffer.h
--- a/Engine/Source/Runtime/AudioMixer/Private/AudioMixerSourceOutputBuffer.h
+++ b/Engine/Source/Runtime/AudioMixer/Private/AudioMixerSourceOutputBuffer.h
@@ -43,12 +43,12 @@ namespace Audio
 		uint32 GetNumSourceChannels() const { return NumSourceChannels; }
 
 		// Sets the channel map
-		bool SetChannelMap(const AlignedFloatBuffer& InChannelMap, bool bInIsCenterChannelOnly);
+		bool SetChannelMap(const FAlignedFloatBuffer& InChannelMap, bool bInIsCenterChannelOnly);
 
 		// Sets the pre and post attenuation source buffers. This is the source buffer data derived from source manager source processing.
-		void SetPreAttenuationSourceBuffer(AlignedFloatBuffer* InPreAttenuationBuffer);
-		void SetPostAttenuationSourceBuffer(AlignedFloatBuffer* InPostAttenuationBuffer);
-		void CopyReverbPluginOutputData(AlignedFloatBuffer& InAudioBuffer);
+		void SetPreAttenuationSourceBuffer(FAlignedFloatBuffer* InPreAttenuationBuffer);
+		void SetPostAttenuationSourceBuffer(FAlignedFloatBuffer* InPostAttenuationBuffer);
+		void CopyReverbPluginOutputData(FAlignedFloatBuffer& InAudioBuffer);
 		const float* GetReverbPluginOutputData() const;
 
 		// Retrieves the current sound field packet for the given key
@@ -59,7 +59,7 @@ namespace Audio
 		void ComputeOutput(const FSpatializationParams& InSpatParams);
 
 		// Called by submixes to mix this output buffer to their buffer
-		void MixOutput(float SendLevel, EMixerSourceSubmixSendStage InSubmixSendStage, AlignedFloatBuffer& OutMixedBuffer) const;
+		void MixOutput(float SendLevel, EMixerSourceSubmixSendStage InSubmixSendStage, FAlignedFloatBuffer& OutMixedBuffer) const;
 
 		// Return the listener rotation
 		// TODO: consolidate the code that is using this to be private to this class.
@@ -140,33 +140,33 @@ namespace Audio
 
 	private: // private methods
 
-		void ComputeOutput3D(AlignedFloatBuffer& InSource, AlignedFloatBuffer& InOutput);
+		void ComputeOutput3D(FAlignedFloatBuffer& InSource, FAlignedFloatBuffer& InOutput);
 		void ComputeOutput3D();
-		void ComputeOutput2D(AlignedFloatBuffer& InSource, AlignedFloatBuffer& InOutput);
+		void ComputeOutput2D(FAlignedFloatBuffer& InSource, FAlignedFloatBuffer& InOutput);
 		void ComputeOutput2D();
 
-		void EncodeSoundfield(FSoundfieldData& InSoundfieldData, Audio::AlignedFloatBuffer& InSourceBuffer);
+		void EncodeSoundfield(FSoundfieldData& InSoundfieldData, Audio::FAlignedFloatBuffer& InSourceBuffer);
 		void EncodeToSoundfieldFormats(const FSpatializationParams& InSpatParams);
 
 	private: // private data
 
 		// Buffer for reverb plugins to write into
-		AlignedFloatBuffer ReverbPluginOutputBuffer;
+		FAlignedFloatBuffer ReverbPluginOutputBuffer;
 
 		// The source buffer created before distance attenuation is applied
-		AlignedFloatBuffer* PreAttenuationSourceBuffer;
+		FAlignedFloatBuffer* PreAttenuationSourceBuffer;
 
 		// The source buffer created after distance attenuation is applied
-		AlignedFloatBuffer* PostAttenuationSourceBuffer;
+		FAlignedFloatBuffer* PostAttenuationSourceBuffer;
 
 		// Data used to map the source channel data to the output channel configuration
 		FSourceChannelMap SourceChannelMap;
 
 		// The result of the source output with source data derived pre distance attenuation
-		AlignedFloatBuffer PreAttenuationOutputBuffer;
+		FAlignedFloatBuffer PreAttenuationOutputBuffer;
 
 		// The result of the source output with source data derived post distance attenuation
-		AlignedFloatBuffer PostAttenuationOutputBuffer;
+		FAlignedFloatBuffer PostAttenuationOutputBuffer;
 
 		// The number of source channels (sound input)
 		uint32 NumSourceChannels;
diff --git a/Engine/Source/Runtime/AudioMixer/Private/AudioMixerSourceVoice.cpp b/Engine/Source/Runtime/AudioMixer/Private/AudioMixerSourceVoice.cpp
--- a/Engine/Source/Runtime/AudioMixer/Private/AudioMixerSourceVoice.cpp
+++ b/Engine/Source/Runtime/AudioMixer/Private/AudioMixerSourceVoice.cpp
@@ -196,7 +196,7 @@ namespace Audio
 		}
 	}
 
-	void FMixerSourceVoice::SetChannelMap(const uint32 NumInputChannels, const Audio::AlignedFloatBuffer& InChannelMap, const bool bInIs3D, const bool bInIsCenterChannelOnly)
+	void FMixerSourceVoice::SetChannelMap(const uint32 NumInputChannels, const Audio::FAlignedFloatBuffer& InChannelMap, const bool bInIs3D, const bool bInIsCenterChannelOnly)
 	{
 		AUDIO_MIXER_CHECK_GAME_THREAD(MixerDevice);
 
@@ -304,7 +304,7 @@ namespace Audio
 		return SourceManager->GetEnvelopeValue(SourceId);
 	}
 
-	void FMixerSourceVoice::MixOutputBuffers(int32 InNumOutputChannels, const float SendLevel, EMixerSourceSubmixSendStage InSubmixSendStage, AlignedFloatBuffer& OutWetBuffer) const
+	void FMixerSourceVoice::MixOutputBuffers(int32 InNumOutputChannels, const float SendLevel, EMixerSourceSubmixSendStage InSubmixSendStage, FAlignedFloatBuffer& OutWetBuffer) const
 	{
 		AUDIO_MIXER_CHECK_AUDIO_PLAT_THREAD(MixerDevice);
 
diff --git a/Engine/Source/Runtime/AudioMixer/Private/AudioMixerSourceVoice.h b/Engine/Source/Runtime/AudioMixer/Private/AudioMixerSourceVoice.h
--- a/Engine/Source/Runtime/AudioMixer/Private/AudioMixerSourceVoice.h
+++ b/Engine/Source/Runtime/AudioMixer/Private/AudioMixerSourceVoice.h
@@ -60,7 +60,7 @@ namespace Audio
 		void SetModHPFFrequency(const float InFrequency);
 
 		// Sets the source voice's channel map (2d or 3d).
-		void SetChannelMap(const uint32 NumInputChannels, const Audio::AlignedFloatBuffer& InChannelMap, const bool bInIs3D, const bool bInIsCenterChannelOnly);
+		void SetChannelMap(const uint32 NumInputChannels, const Audio::FAlignedFloatBuffer& InChannelMap, const bool bInIs3D, const bool bInIsCenterChannelOnly);
 
 		// Sets params used by HRTF spatializer
 		void SetSpatializationParams(const FSpatializationParams& InParams);
@@ -107,7 +107,7 @@ namespace Audio
 		float GetEnvelopeValue() const;
 
 		// Mixes the dry and wet buffer audio into the given buffers.
-		void MixOutputBuffers(int32 InNumChannels, const float SendLevel, EMixerSourceSubmixSendStage InSubmixSendStage, AlignedFloatBuffer& OutWetBuffer) const;
+		void MixOutputBuffers(int32 InNumChannels, const float SendLevel, EMixerSourceSubmixSendStage InSubmixSendStage, FAlignedFloatBuffer& OutWetBuffer) const;
 
 		// For soundfield conversions, get the encoded audio.
 		const ISoundfieldAudioPacket* GetEncodedOutput(const FSoundfieldEncodingKey& InKey) const;
diff --git a/Engine/Source/Runtime/AudioMixer/Private/AudioMixerSubmix.cpp b/Engine/Source/Runtime/AudioMixer/Private/AudioMixerSubmix.cpp
--- a/Engine/Source/Runtime/AudioMixer/Private/AudioMixerSubmix.cpp
+++ b/Engine/Source/Runtime/AudioMixer/Private/AudioMixerSubmix.cpp
@@ -355,9 +355,9 @@ namespace Audio
 		}
 	}
 
-	void FMixerSubmix::DownmixBuffer(const int32 InChannels, const AlignedFloatBuffer& InBuffer, const int32 OutChannels, AlignedFloatBuffer& OutNewBuffer)
+	void FMixerSubmix::DownmixBuffer(const int32 InChannels, const FAlignedFloatBuffer& InBuffer, const int32 OutChannels, FAlignedFloatBuffer& OutNewBuffer)
 	{
-		Audio::AlignedFloatBuffer MixdownGainsMap;
+		Audio::FAlignedFloatBuffer MixdownGainsMap;
 		Audio::FMixerDevice::Get2DChannelMap(false, InChannels, OutChannels, false, MixdownGainsMap);
 		Audio::DownmixBuffer(InChannels, OutChannels, InBuffer, OutNewBuffer, MixdownGainsMap.GetData());
 	}
@@ -670,7 +670,7 @@ namespace Audio
 		});
 	}
 
-	void FMixerSubmix::MixBufferDownToMono(const AlignedFloatBuffer& InBuffer, int32 NumInputChannels, AlignedFloatBuffer& OutBuffer)
+	void FMixerSubmix::MixBufferDownToMono(const FAlignedFloatBuffer& InBuffer, int32 NumInputChannels, FAlignedFloatBuffer& OutBuffer)
 	{
 		check(NumInputChannels > 0);
 
@@ -975,7 +975,7 @@ namespace Audio
 		return OwningSubmixObject.IsValid();
 	}
 
-	void FMixerSubmix::ProcessAudio(AlignedFloatBuffer& OutAudioBuffer)
+	void FMixerSubmix::ProcessAudio(FAlignedFloatBuffer& OutAudioBuffer)
 	{
 		AUDIO_MIXER_CHECK_AUDIO_PLAT_THREAD(MixerDevice);
 
@@ -1285,7 +1285,7 @@ namespace Audio
 		}
 	}
 
-	bool FMixerSubmix::GenerateEffectChainAudio(FSoundEffectSubmixInputData& InputData, AlignedFloatBuffer& InAudioBuffer, TArray<FSoundEffectSubmixPtr>& InEffectChain, AlignedFloatBuffer& OutBuffer)
+	bool FMixerSubmix::GenerateEffectChainAudio(FSoundEffectSubmixInputData& InputData, FAlignedFloatBuffer& InAudioBuffer, TArray<FSoundEffectSubmixPtr>& InEffectChain, FAlignedFloatBuffer& OutBuffer)
 	{
 		// Reset the output scratch buffer
 		ScratchBuffer.Reset(NumSamples);
@@ -1752,7 +1752,7 @@ namespace Audio
 		bIsRecording = true;
 	}
 
-	AlignedFloatBuffer& FMixerSubmix::OnStopRecordingOutput(float& OutNumChannels, float& OutSampleRate)
+	FAlignedFloatBuffer& FMixerSubmix::OnStopRecordingOutput(float& OutNumChannels, float& OutSampleRate)
 	{
 		FScopeLock ScopedLock(&RecordingCriticalSection);
 		bIsRecording = false;
diff --git a/Engine/Source/Runtime/AudioMixer/Private/Effects/AudioMixerSubmixEffectEQ.cpp b/Engine/Source/Runtime/AudioMixer/Private/Effects/AudioMixerSubmixEffectEQ.cpp
--- a/Engine/Source/Runtime/AudioMixer/Private/Effects/AudioMixerSubmixEffectEQ.cpp
+++ b/Engine/Source/Runtime/AudioMixer/Private/Effects/AudioMixerSubmixEffectEQ.cpp
@@ -100,8 +100,8 @@ void FSubmixEffectSubmixEQ::OnProcessAudio(const FSoundEffectSubmixInputData& In
 	// Update parameters that may have been set from game thread
 	UpdateParameters(InData.NumChannels);
 
-	Audio::AlignedFloatBuffer& InAudioBuffer = *InData.AudioBuffer;
-	Audio::AlignedFloatBuffer& OutAudioBuffer = *OutData.AudioBuffer;
+	Audio::FAlignedFloatBuffer& InAudioBuffer = *InData.AudioBuffer;
+	Audio::FAlignedFloatBuffer& OutAudioBuffer = *OutData.AudioBuffer;
 
 	if (bEQSettingsSet && RenderThreadEQSettings.EQBands.Num() > 0)
 	{
diff --git a/Engine/Source/Runtime/AudioMixer/Private/SoundWaveDecoder.cpp b/Engine/Source/Runtime/AudioMixer/Private/SoundWaveDecoder.cpp
--- a/Engine/Source/Runtime/AudioMixer/Private/SoundWaveDecoder.cpp
+++ b/Engine/Source/Runtime/AudioMixer/Private/SoundWaveDecoder.cpp
@@ -277,7 +277,7 @@ namespace Audio
 		}
 	}
 
-	void FDecodingSoundSource::GetAudioBufferInternal(const int32 InNumFrames, const int32 InNumChannels, AlignedFloatBuffer& OutAudioBuffer)
+	void FDecodingSoundSource::GetAudioBufferInternal(const int32 InNumFrames, const int32 InNumChannels, FAlignedFloatBuffer& OutAudioBuffer)
 	{
 #if AUDIO_SOURCE_DECODER_DEBUG
 		int32 SampleIndex = 0;
@@ -363,7 +363,7 @@ namespace Audio
 #endif
 	}
 
-	bool FDecodingSoundSource::GetAudioBuffer(const int32 InNumFrames, const int32 InNumChannels, AlignedFloatBuffer& OutAudioBuffer)
+	bool FDecodingSoundSource::GetAudioBuffer(const int32 InNumFrames, const int32 InNumChannels, FAlignedFloatBuffer& OutAudioBuffer)
 	{
 		FScopeTryLock Lock(&MixerSourceBufferCritSec);
 
@@ -684,7 +684,7 @@ namespace Audio
 	}
 
 
-	bool FSoundSourceDecoder::GetSourceBuffer(const FDecodingSoundSourceHandle& InHandle, const int32 NumOutFrames, const int32 NumOutChannels, AlignedFloatBuffer& OutAudioBuffer)
+	bool FSoundSourceDecoder::GetSourceBuffer(const FDecodingSoundSourceHandle& InHandle, const int32 NumOutFrames, const int32 NumOutChannels, FAlignedFloatBuffer& OutAudioBuffer)
 	{
 		check(InHandle.Id != INDEX_NONE);
 		FScopeLock Lock(&DecodingSourcesCritSec);
diff --git a/Engine/Source/Runtime/AudioMixer/Public/AudioMixerBus.h b/Engine/Source/Runtime/AudioMixer/Public/AudioMixerBus.h
--- a/Engine/Source/Runtime/AudioMixer/Public/AudioMixerBus.h
+++ b/Engine/Source/Runtime/AudioMixer/Public/AudioMixerBus.h
@@ -74,8 +74,8 @@ namespace Audio
 		void MixBuffer();
 
 		// Copies the current internal buffer to a provided output buffer. Only supports mono or stereo input/output formats.
-		void CopyCurrentBuffer(Audio::AlignedFloatBuffer& InChannelMap, int32 InNumOutputChannels, AlignedFloatBuffer& OutBuffer, int32 NumOutputFrames) const;
-		void CopyCurrentBuffer(int32 InNumOutputChannels, AlignedFloatBuffer& OutBuffer, int32 NumOutputFrames) const;
+		void CopyCurrentBuffer(Audio::FAlignedFloatBuffer& InChannelMap, int32 InNumOutputChannels, FAlignedFloatBuffer& OutBuffer, int32 NumOutputFrames) const;
+		void CopyCurrentBuffer(int32 InNumOutputChannels, FAlignedFloatBuffer& OutBuffer, int32 NumOutputFrames) const;
 
 		// If this bus was constructed before
 		void SetNumOutputChannels(int32 InNumOutputChannels);
diff --git a/Engine/Source/Runtime/AudioMixer/Public/AudioMixerDevice.h b/Engine/Source/Runtime/AudioMixer/Public/AudioMixerDevice.h
--- a/Engine/Source/Runtime/AudioMixer/Public/AudioMixerDevice.h
+++ b/Engine/Source/Runtime/AudioMixer/Public/AudioMixerDevice.h
@@ -141,7 +141,7 @@ namespace Audio
 
 		// Submix recording callbacks:
 		virtual void StartRecording(USoundSubmix* InSubmix, float ExpectedRecordingDuration) override;
-		virtual Audio::AlignedFloatBuffer& StopRecording(USoundSubmix* InSubmix, float& OutNumChannels, float& OutSampleRate) override;
+		virtual Audio::FAlignedFloatBuffer& StopRecording(USoundSubmix* InSubmix, float& OutNumChannels, float& OutSampleRate) override;
 
 		virtual void PauseRecording(USoundSubmix* InSubmix);
 		virtual void ResumeRecording(USoundSubmix* InSubmix);
@@ -174,7 +174,7 @@ namespace Audio
 		//~ End FAudioDevice
 
 		//~ Begin IAudioMixer
-		virtual bool OnProcessAudioStream(AlignedFloatBuffer& OutputBuffer) override;
+		virtual bool OnProcessAudioStream(FAlignedFloatBuffer& OutputBuffer) override;
 		virtual void OnAudioStreamShutdown() override;
 		//~ End IAudioMixer
 
@@ -216,11 +216,11 @@ namespace Audio
 		IAudioMixerPlatformInterface* GetAudioMixerPlatform() const { return AudioMixerPlatform; }
 
 		// Builds a 3D channel map for a spatialized source.
-		void Get3DChannelMap(const int32 InSubmixNumChannels, const FWaveInstance* InWaveInstance, const float EmitterAzimuth, const float NormalizedOmniRadius, Audio::AlignedFloatBuffer& OutChannelMap);
+		void Get3DChannelMap(const int32 InSubmixNumChannels, const FWaveInstance* InWaveInstance, const float EmitterAzimuth, const float NormalizedOmniRadius, Audio::FAlignedFloatBuffer& OutChannelMap);
 
 		// Builds a channel gain matrix for a non-spatialized source. The non-static variation of this function queries AudioMixerDevice->NumOutputChannels directly which may not be thread safe.
-		void Get2DChannelMap(bool bIsVorbis, const int32 NumSourceChannels, const bool bIsCenterChannelOnly, Audio::AlignedFloatBuffer& OutChannelMap) const;
-		static void Get2DChannelMap(bool bIsVorbis, const int32 NumSourceChannels, const int32 NumOutputChannels, const bool bIsCenterChannelOnly, Audio::AlignedFloatBuffer& OutChannelMap);
+		void Get2DChannelMap(bool bIsVorbis, const int32 NumSourceChannels, const bool bIsCenterChannelOnly, Audio::FAlignedFloatBuffer& OutChannelMap) const;
+		static void Get2DChannelMap(bool bIsVorbis, const int32 NumSourceChannels, const int32 NumOutputChannels, const bool bIsCenterChannelOnly, Audio::FAlignedFloatBuffer& OutChannelMap);
 
 		int32 GetDeviceSampleRate() const;
 		int32 GetDeviceOutputChannels() const;
@@ -311,8 +311,8 @@ namespace Audio
 		void CacheChannelMap(const int32 NumSourceChannels, const int32 NumOutputChannels, const bool bIsCenterChannelOnly);
 		void InitializeChannelAzimuthMap(const int32 NumChannels);
 
-		void WhiteNoiseTest(AlignedFloatBuffer& Output);
-		void SineOscTest(AlignedFloatBuffer& Output);
+		void WhiteNoiseTest(FAlignedFloatBuffer& Output);
+		void SineOscTest(FAlignedFloatBuffer& Output);
 
 		bool IsMainAudioDevice() const;
 
diff --git a/Engine/Source/Runtime/AudioMixer/Public/AudioMixerSubmix.h b/Engine/Source/Runtime/AudioMixer/Public/AudioMixerSubmix.h
--- a/Engine/Source/Runtime/AudioMixer/Public/AudioMixerSubmix.h
+++ b/Engine/Source/Runtime/AudioMixer/Public/AudioMixerSubmix.h
@@ -190,7 +190,7 @@ namespace Audio
 		bool IsValid() const;
 
 		// Function which processes audio.
-		void ProcessAudio(AlignedFloatBuffer& OutAudio);
+		void ProcessAudio(FAlignedFloatBuffer& OutAudio);
 		void ProcessAudio(ISoundfieldAudioPacket& OutputAudio);
 
 		// This should be called if this submix doesn't send it's audio to a parent submix,
@@ -226,7 +226,7 @@ namespace Audio
 		void OnStartRecordingOutput(float ExpectedDuration);
 
 		// This is called by the corresponding USoundSubmix when StopRecordingOutput is called.
-		AlignedFloatBuffer& OnStopRecordingOutput(float& OutNumChannels, float& OutSampleRate);
+		FAlignedFloatBuffer& OnStopRecordingOutput(float& OutNumChannels, float& OutSampleRate);
 
 		// This is called by the corresponding USoundSubmix when PauseRecording is called.
 		void PauseRecordingOutput();
@@ -301,9 +301,9 @@ namespace Audio
 		void InitInternal();
 
 		// Down mix the given buffer to the desired down mix channel count
-		static void DownmixBuffer(const int32 InChannels, const AlignedFloatBuffer& InBuffer, const int32 OutChannels, AlignedFloatBuffer& OutNewBuffer);
+		static void DownmixBuffer(const int32 InChannels, const FAlignedFloatBuffer& InBuffer, const int32 OutChannels, FAlignedFloatBuffer& OutNewBuffer);
 
-		void MixBufferDownToMono(const AlignedFloatBuffer& InBuffer, int32 NumInputChannels, AlignedFloatBuffer& OutBuffer);
+		void MixBufferDownToMono(const FAlignedFloatBuffer& InBuffer, int32 NumInputChannels, FAlignedFloatBuffer& OutBuffer);
 
 		void SetupSoundfieldEncodersForChildren();
 		void SetupSoundfieldEncodingForChild(FChildSubmixInfo& InChild);
@@ -335,7 +335,7 @@ namespace Audio
 		void SubmixCommand(TFunction<void()> Command);
 
 		// Generates audio from the given effect chain into the given buffer
-		bool GenerateEffectChainAudio(FSoundEffectSubmixInputData& InputData, AlignedFloatBuffer& InAudioBuffer, TArray<FSoundEffectSubmixPtr>& InEffectChain, AlignedFloatBuffer& OutBuffer);
+		bool GenerateEffectChainAudio(FSoundEffectSubmixInputData& InputData, FAlignedFloatBuffer& InAudioBuffer, TArray<FSoundEffectSubmixPtr>& InEffectChain, FAlignedFloatBuffer& OutBuffer);
 
 		// This mixer submix's Id
 		uint32 Id;
@@ -362,7 +362,7 @@ namespace Audio
 
 		// The array of submix effect overrides. There may be more than one if multiple are fading out. There should be only one fading in (the current override).
 		TArray<FSubmixEffectFadeInfo> EffectChains;
-		AlignedFloatBuffer EffectChainOutputBuffer;
+		FAlignedFloatBuffer EffectChainOutputBuffer;
 
 		// Owning mixer device. 
 		FMixerDevice* MixerDevice;
@@ -370,11 +370,11 @@ namespace Audio
 		// Map of mixer source voices with a given send level for this submix
 		TMap<FMixerSourceVoice*, FSubmixVoiceData> MixerSourceVoices;
 
-		AlignedFloatBuffer ScratchBuffer;
-		AlignedFloatBuffer SubmixChainMixBuffer;
-		AlignedFloatBuffer InputBuffer;
-		AlignedFloatBuffer DownmixedBuffer;
-		AlignedFloatBuffer SourceInputBuffer;
+		FAlignedFloatBuffer ScratchBuffer;
+		FAlignedFloatBuffer SubmixChainMixBuffer;
+		FAlignedFloatBuffer InputBuffer;
+		FAlignedFloatBuffer DownmixedBuffer;
+		FAlignedFloatBuffer SourceInputBuffer;
 
 		int32 NumChannels;
 		int32 NumSamples;
@@ -462,10 +462,10 @@ namespace Audio
 			TUniquePtr<ISoundfieldEndpoint> SoundfieldEndpoint;
 
 			// for non-soundfield endpoints, we use these buffers for processing.
-			AlignedFloatBuffer AudioBuffer;
-			AlignedFloatBuffer ResampledAudioBuffer;
-			AlignedFloatBuffer DownmixedResampledAudioBuffer;
-			AlignedFloatBuffer DownmixChannelMap;
+			FAlignedFloatBuffer AudioBuffer;
+			FAlignedFloatBuffer ResampledAudioBuffer;
+			FAlignedFloatBuffer DownmixedResampledAudioBuffer;
+			FAlignedFloatBuffer DownmixChannelMap;
 
 			// Number of channels and sample rate for the external endpoint.
 			int32 NumChannels;
@@ -516,10 +516,10 @@ namespace Audio
 		TSharedPtr<FAsyncSpectrumAnalyzer, ESPMode::ThreadSafe> SpectrumAnalyzer;
 		
 		// This buffer is used to downmix the submix output to mono before submitting it to the SpectrumAnalyzer.
-		AlignedFloatBuffer MonoMixBuffer;
+		FAlignedFloatBuffer MonoMixBuffer;
 
 		// The dry channel buffer
-		AlignedFloatBuffer DryChannelBuffer;
+		FAlignedFloatBuffer DryChannelBuffer;
 
 		// Submix command queue to shuffle commands from audio thread to audio render thread.
 		TQueue<TFunction<void()>> CommandQueue;
@@ -531,7 +531,7 @@ namespace Audio
 		FCriticalSection BufferListenerCriticalSection;
 
 		// This buffer is used for recorded output of the submix.
-		AlignedFloatBuffer RecordingData;
+		FAlignedFloatBuffer RecordingData;
 
 		// Returns the number of submix effects
 		int32 NumSubmixEffects;
diff --git a/Engine/Source/Runtime/AudioMixer/Public/SoundWaveDecoder.h b/Engine/Source/Runtime/AudioMixer/Public/SoundWaveDecoder.h
--- a/Engine/Source/Runtime/AudioMixer/Public/SoundWaveDecoder.h
+++ b/Engine/Source/Runtime/AudioMixer/Public/SoundWaveDecoder.h
@@ -80,7 +80,7 @@ namespace Audio
 		void SetForceSyncDecode(bool bShouldForceSyncDecode);
 
 		// Get audio buffer
-		bool GetAudioBuffer(const int32 InNumFrames, const int32 InNumChannels, AlignedFloatBuffer& OutAudioBuffer);
+		bool GetAudioBuffer(const int32 InNumFrames, const int32 InNumChannels, FAlignedFloatBuffer& OutAudioBuffer);
 
 		// Return the underlying sound wave
 		USoundWave* GetSoundWave() { return SoundWave; }
@@ -88,7 +88,7 @@ namespace Audio
 	private:
 
 		void ReadFrame();
-		void GetAudioBufferInternal(const int32 InNumFrames, const int32 InNumChannels, AlignedFloatBuffer& OutAudioBuffer);
+		void GetAudioBufferInternal(const int32 InNumFrames, const int32 InNumChannels, FAlignedFloatBuffer& OutAudioBuffer);
 
 		// Handle to the decoding source
 		FDecodingSoundSourceHandle Handle;
@@ -106,7 +106,7 @@ namespace Audio
 		TSharedPtr<FMixerSourceBuffer, ESPMode::ThreadSafe> MixerSourceBuffer;
 
 		// Scratch buffer used for upmixing and downmixing the audio
-		AlignedFloatBuffer ScratchBuffer;
+		FAlignedFloatBuffer ScratchBuffer;
 
 		// Sample rate of the source
 		int32 SampleRate;
@@ -250,7 +250,7 @@ namespace Audio
 		void SetSourceVolumeScale(const FDecodingSoundSourceHandle& Handle, float InVolumeScale);
 
 		// Get a decoded buffer for the given decoding sound wave handle. Call only from audio render thread or audio render thread task.
-		bool GetSourceBuffer(const FDecodingSoundSourceHandle& InHandle, const int32 NumOutFrames, const int32 NumOutChannels, AlignedFloatBuffer& OutAudioBuffer);
+		bool GetSourceBuffer(const FDecodingSoundSourceHandle& InHandle, const int32 NumOutFrames, const int32 NumOutChannels, FAlignedFloatBuffer& OutAudioBuffer);
 
 		// Queries if the decoding source is finished
 		bool IsFinished(const FDecodingSoundSourceHandle& InHandle) const;
diff --git a/Engine/Source/Runtime/AudioMixerCore/Public/AudioMixer.h b/Engine/Source/Runtime/AudioMixerCore/Public/AudioMixer.h
--- a/Engine/Source/Runtime/AudioMixerCore/Public/AudioMixer.h
+++ b/Engine/Source/Runtime/AudioMixerCore/Public/AudioMixer.h
@@ -168,7 +168,7 @@ namespace Audio
 	{
 	public:
 		/** Callback to generate a new audio stream buffer. */
-		virtual bool OnProcessAudioStream(AlignedFloatBuffer& OutputBuffer) = 0;
+		virtual bool OnProcessAudioStream(FAlignedFloatBuffer& OutputBuffer) = 0;
 
 		/** Called when audio render thread stream is shutting down. Last function called. Allows cleanup on render thread. */
 		virtual void OnAudioStreamShutdown() = 0;
@@ -332,13 +332,13 @@ namespace Audio
 		mutable Audio::TCircularAudioBuffer<uint8> CircularBuffer;
 		
 		// Buffer that we render audio to from the IAudioMixer instance associated with this output buffer.
-		Audio::AlignedFloatBuffer RenderBuffer;
+		Audio::FAlignedFloatBuffer RenderBuffer;
 
 		// Buffer read by the platform interface thread.
-		mutable Audio::AlignedByteBuffer PopBuffer;
+		mutable Audio::FAlignedByteBuffer PopBuffer;
 
 		// For non-float situations, this buffer is used to convert RenderBuffer before pushing it to CircularBuffer.
-		AlignedByteBuffer FormattedBuffer;
+		FAlignedByteBuffer FormattedBuffer;
  		EAudioMixerStreamDataFormat::Type DataFormat;
 
 		static size_t GetSizeForDataFormat(EAudioMixerStreamDataFormat::Type InDataFormat);
diff --git a/Engine/Source/Runtime/AudioPlatformConfiguration/Public/AudioResampler.h b/Engine/Source/Runtime/AudioPlatformConfiguration/Public/AudioResampler.h
--- a/Engine/Source/Runtime/AudioPlatformConfiguration/Public/AudioResampler.h
+++ b/Engine/Source/Runtime/AudioPlatformConfiguration/Public/AudioResampler.h
@@ -1,11 +1,12 @@
 // Copyright Epic Games, Inc. All Rights Reserved.
-
 #pragma once
-//#if PLATFORM_WINDOWS
+
 #include "CoreMinimal.h"
 
 #if PLATFORM_SWITCH
 // Switch uses page alignment for submitted buffers
+// TODO: Move BufferVectorOperations AUDIO_BUFFER_ALIGNMENT
+// define to more central location and reference here.
 #define AUDIO_BUFFER_ALIGNMENT 4096
 #else
 #define AUDIO_BUFFER_ALIGNMENT 16
@@ -15,8 +16,12 @@ DECLARE_LOG_CATEGORY_EXTERN(LogAudioResampler, Warning, All);
 
 namespace Audio
 {
-	typedef TArray<float, TAlignedHeapAllocator<AUDIO_BUFFER_ALIGNMENT>> AlignedFloatBuffer;
-	typedef TArray<uint8, TAlignedHeapAllocator<AUDIO_BUFFER_ALIGNMENT>> AlignedByteBuffer;
+	// TODO: Move BufferVectorOperations Aligned...Buffer
+	// defines to more central location and reference here.
+	namespace VectorOps
+	{
+		using FAlignedFloatBuffer = TArray<float, TAlignedHeapAllocator<AUDIO_BUFFER_ALIGNMENT>>;
+	} // namespace VectorOps
 
 	enum class EResamplingMethod : uint8
 	{
@@ -33,12 +38,12 @@ namespace Audio
 		int32 NumChannels;
 		float SourceSampleRate;
 		float DestinationSampleRate;
-		AlignedFloatBuffer& InputBuffer;
+		VectorOps::FAlignedFloatBuffer& InputBuffer;
 	};
 
 	struct FResamplerResults
 	{
-		AlignedFloatBuffer* OutBuffer;
+		VectorOps::FAlignedFloatBuffer* OutBuffer;
 
 		float ResultingSampleRate;
 
@@ -78,6 +83,4 @@ namespace Audio
 		TUniquePtr<FResamplerImpl> Impl;
 	};
 	
-}
-
-//#endif //PLATFORM_WINDOWS
\ No newline at end of file
+} // namespace Audio
diff --git a/Engine/Source/Runtime/Engine/Classes/Sound/SoundEffectSubmix.h b/Engine/Source/Runtime/Engine/Classes/Sound/SoundEffectSubmix.h
--- a/Engine/Source/Runtime/Engine/Classes/Sound/SoundEffectSubmix.h
+++ b/Engine/Source/Runtime/Engine/Classes/Sound/SoundEffectSubmix.h
@@ -4,6 +4,7 @@
 #include "CoreMinimal.h"
 
 #include "AudioDefines.h"
+#include "DSP/BufferVectorOperations.h"
 #include "IAudioExtensionPlugin.h"
 #include "Sound/SoundEffectPreset.h"
 #include "Sound/SoundEffectBase.h"
@@ -57,7 +58,7 @@ struct FSoundEffectSubmixInputData
 	const TArray<FTransform>* ListenerTransforms;
 
 	/** The raw input audio buffer. Size is NumFrames * NumChannels */
-	Audio::AlignedFloatBuffer* AudioBuffer;
+	Audio::FAlignedFloatBuffer* AudioBuffer;
 
 	/** Sample accurate audio clock. */
 	double AudioClock;
@@ -75,7 +76,7 @@ struct FSoundEffectSubmixInputData
 struct FSoundEffectSubmixOutputData
 {
 	/** The output audio buffer. */
-	Audio::AlignedFloatBuffer* AudioBuffer;
+	Audio::FAlignedFloatBuffer* AudioBuffer;
 
 	/** The number of channels of the submix. */
 	int32 NumChannels;
diff --git a/Engine/Source/Runtime/Engine/Private/AudioDerivedData.cpp b/Engine/Source/Runtime/Engine/Private/AudioDerivedData.cpp
--- a/Engine/Source/Runtime/Engine/Private/AudioDerivedData.cpp
+++ b/Engine/Source/Runtime/Engine/Private/AudioDerivedData.cpp
@@ -956,8 +956,8 @@ static void ResampleWaveData(TArray<uint8>& WaveData, size_t& NumBytes, int32 Nu
 	double StartTime = FPlatformTime::Seconds();
 
 	// Set up temporary output buffers:
-	Audio::AlignedFloatBuffer ResamplerInputData;
-	Audio::AlignedFloatBuffer ResamplerOutputData;
+	Audio::FAlignedFloatBuffer ResamplerInputData;
+	Audio::FAlignedFloatBuffer ResamplerOutputData;
 
 	int32 NumSamples = NumBytes / sizeof(int16);
 
diff --git a/Engine/Source/Runtime/Engine/Private/SoundWave.cpp b/Engine/Source/Runtime/Engine/Private/SoundWave.cpp
--- a/Engine/Source/Runtime/Engine/Private/SoundWave.cpp
+++ b/Engine/Source/Runtime/Engine/Private/SoundWave.cpp
@@ -1621,7 +1621,7 @@ void USoundWave::BakeFFTAnalysis()
 		Audio::FSpectrumAnalyzer SpectrumAnalyzer(SpectrumAnalyzerSettings, (float)RawDataSampleRate);
 
 		// The audio data block to use to submit audio data to the spectrum analyzer
-		Audio::AlignedFloatBuffer AnalysisData;
+		Audio::FAlignedFloatBuffer AnalysisData;
 		check(FFTAnalysisFrameSize > 256);
 		AnalysisData.Reserve(FFTAnalysisFrameSize);
 
diff --git a/Engine/Source/Runtime/Engine/Public/AudioDevice.h b/Engine/Source/Runtime/Engine/Public/AudioDevice.h
--- a/Engine/Source/Runtime/Engine/Public/AudioDevice.h
+++ b/Engine/Source/Runtime/Engine/Public/AudioDevice.h
@@ -1302,11 +1302,11 @@ public:
 	}
 
 	/** This is called by a USoundSubmix when we stop recording a submix on this device. */
-	virtual Audio::AlignedFloatBuffer& StopRecording(USoundSubmix* InSubmix, float& OutNumChannels, float& OutSampleRate)
+	virtual Audio::FAlignedFloatBuffer& StopRecording(USoundSubmix* InSubmix, float& OutNumChannels, float& OutSampleRate)
 	{
 		UE_LOG(LogAudio, Error, TEXT("Submix recording only works with the audio mixer. Please run using -audiomixer to or set INI file use submix recording."));
 
-		static Audio::AlignedFloatBuffer InvalidBuffer;
+		static Audio::FAlignedFloatBuffer InvalidBuffer;
 		return InvalidBuffer;
 	}
 
diff --git a/Engine/Source/Runtime/SignalProcessing/Private/AllPassFractionalDelay.cpp b/Engine/Source/Runtime/SignalProcessing/Private/AllPassFractionalDelay.cpp
--- a/Engine/Source/Runtime/SignalProcessing/Private/AllPassFractionalDelay.cpp
+++ b/Engine/Source/Runtime/SignalProcessing/Private/AllPassFractionalDelay.cpp
@@ -59,7 +59,7 @@ void FAllPassFractionalDelay::Reset()
 }
 
 
-void FAllPassFractionalDelay::ProcessAudio(const AlignedFloatBuffer& InSamples, const AlignedFloatBuffer& InDelays, AlignedFloatBuffer& OutSamples)
+void FAllPassFractionalDelay::ProcessAudio(const FAlignedFloatBuffer& InSamples, const FAlignedFloatBuffer& InDelays, FAlignedFloatBuffer& OutSamples)
 {
 	const int32 InNum = InSamples.Num();
 	checkf(InNum == InDelays.Num(), TEXT("Input buffers must be equal length."));
diff --git a/Engine/Source/Runtime/SignalProcessing/Private/AudioChannelFormatConverter.cpp b/Engine/Source/Runtime/SignalProcessing/Private/AudioChannelFormatConverter.cpp
--- a/Engine/Source/Runtime/SignalProcessing/Private/AudioChannelFormatConverter.cpp
+++ b/Engine/Source/Runtime/SignalProcessing/Private/AudioChannelFormatConverter.cpp
@@ -118,7 +118,7 @@ namespace Audio
 		return OutputGainState.bFadeToNextGain ? OutputGainState.NextGain : OutputGainState.Gain;
 	}
 
-	void FBaseChannelFormatConverter::ProcessAudio(const TArray<AlignedFloatBuffer>& InInputBuffers, TArray<AlignedFloatBuffer>& OutOutputBuffers)
+	void FBaseChannelFormatConverter::ProcessAudio(const TArray<FAlignedFloatBuffer>& InInputBuffers, TArray<FAlignedFloatBuffer>& OutOutputBuffers)
 	{
 		using FMixElement = TSortedMap<FChannelMixKey, FChannelMixState>::ElementType;
 
@@ -147,7 +147,7 @@ namespace Audio
 		// is used to determine whether the input channel in a mix entry differs
 		// from the previous. If so, the "Current" input channel is updated.
 		int32 CurrentInputChannelIndex = INDEX_NONE;
-		const AlignedFloatBuffer* CurrentInputChannel = nullptr;
+		const FAlignedFloatBuffer* CurrentInputChannel = nullptr;
 
 		TArray<FChannelMixKey> EntriesToRemove;
 
@@ -161,7 +161,7 @@ namespace Audio
 			if (MixState.InputChannelIndex != CurrentInputChannelIndex)
 			{
 				// Check validity of input channel.
-				const AlignedFloatBuffer& InputBuffer = InInputBuffers[MixState.InputChannelIndex];
+				const FAlignedFloatBuffer& InputBuffer = InInputBuffers[MixState.InputChannelIndex];
 
 				if (ensure(InputBuffer.Num() == NumFramesPerCall))
 				{
diff --git a/Engine/Source/Runtime/SignalProcessing/Private/AudioFFT.cpp b/Engine/Source/Runtime/SignalProcessing/Private/AudioFFT.cpp
--- a/Engine/Source/Runtime/SignalProcessing/Private/AudioFFT.cpp
+++ b/Engine/Source/Runtime/SignalProcessing/Private/AudioFFT.cpp
@@ -511,8 +511,8 @@ namespace Audio
 			const int32 NumOutputFFTElements;
 			FFTFreqDomainData FreqDomainData;
 
-			AlignedFloatBuffer FreqRealBuffer;
-			AlignedFloatBuffer FreqImagBuffer;
+			FAlignedFloatBuffer FreqRealBuffer;
+			FAlignedFloatBuffer FreqImagBuffer;
 			
 
 		public:
@@ -652,7 +652,7 @@ namespace Audio
 		return MakeUnique<FAudioFFTAlgorithm>(FFTSize);
 	}
 
-	void ComputePowerSpectrumNoScaling(const FFTFreqDomainData& InFrequencyData, int32 FFTSize, AlignedFloatBuffer& OutBuffer)
+	void ComputePowerSpectrumNoScaling(const FFTFreqDomainData& InFrequencyData, int32 FFTSize, FAlignedFloatBuffer& OutBuffer)
 	{
 		check((FFTSize % 2) == 0);
 
@@ -695,7 +695,7 @@ namespace Audio
 		}
 	}
 
-	void ComputePowerSpectrum(const FFTFreqDomainData& InFrequencyData, int32 FFTSize, AlignedFloatBuffer& OutBuffer)
+	void ComputePowerSpectrum(const FFTFreqDomainData& InFrequencyData, int32 FFTSize, FAlignedFloatBuffer& OutBuffer)
 	{
 		if (FFTSize < 1)
 		{
@@ -723,7 +723,7 @@ namespace Audio
 		OutBufferData[NumSpectrumValues - 1] *= FFTScale;
 	}
 
-	void ComputeMagnitudeSpectrum(const FFTFreqDomainData& InFrequencyData, int32 FFTSize, AlignedFloatBuffer& OutBuffer)
+	void ComputeMagnitudeSpectrum(const FFTFreqDomainData& InFrequencyData, int32 FFTSize, FAlignedFloatBuffer& OutBuffer)
 	{
 		if (FFTSize < 1)
 		{
@@ -756,7 +756,7 @@ namespace Audio
 		OutBufferData[NumSpectrumValues - 1] *= FFTScale;
 	}
 
-	void ComputeSpectrum(ESpectrumType InSpectrumType, const FFTFreqDomainData& InFrequencyData, int32 FFTSize, AlignedFloatBuffer& OutBuffer)
+	void ComputeSpectrum(ESpectrumType InSpectrumType, const FFTFreqDomainData& InFrequencyData, int32 FFTSize, FAlignedFloatBuffer& OutBuffer)
 	{
 		switch (InSpectrumType)
 		{
@@ -772,7 +772,7 @@ namespace Audio
 				checkf(false, TEXT("Unhandled Audio::ESpectrumType"));
 		}
 	}
-	SIGNALPROCESSING_API void CrossCorrelate(AlignedFloatBuffer& FirstBuffer, AlignedFloatBuffer& SecondBuffer, AlignedFloatBuffer& OutCorrelation, bool bZeroPad /*= true*/)
+	SIGNALPROCESSING_API void CrossCorrelate(FAlignedFloatBuffer& FirstBuffer, FAlignedFloatBuffer& SecondBuffer, FAlignedFloatBuffer& OutCorrelation, bool bZeroPad /*= true*/)
 	{
 		FrequencyBuffer OutputCorrelationFrequencies;
 		CrossCorrelate(FirstBuffer, SecondBuffer, OutputCorrelationFrequencies, bZeroPad);
@@ -796,7 +796,7 @@ namespace Audio
 		PerformIFFT(FreqDomainData, TimeDomainData);
 	}
 
-	SIGNALPROCESSING_API void CrossCorrelate(AlignedFloatBuffer& FirstBuffer, AlignedFloatBuffer& SecondBuffer, FrequencyBuffer& OutCorrelation, bool bZeroPad /*= true*/)
+	SIGNALPROCESSING_API void CrossCorrelate(FAlignedFloatBuffer& FirstBuffer, FAlignedFloatBuffer& SecondBuffer, FrequencyBuffer& OutCorrelation, bool bZeroPad /*= true*/)
 	{
 		const int32 NumSamples = FMath::Max(FirstBuffer.Num(), SecondBuffer.Num());
 
@@ -1048,7 +1048,7 @@ namespace Audio
 		COLABuffer.AddZeroed(BlockSize);
 	}
 
-	void FFFTConvolver::SetFilter(const AlignedFloatBuffer& TimeDomainBuffer)
+	void FFFTConvolver::SetFilter(const FAlignedFloatBuffer& TimeDomainBuffer)
 	{
 		SetFilter(TimeDomainBuffer.GetData(), TimeDomainBuffer.Num());
 	}
diff --git a/Engine/Source/Runtime/SignalProcessing/Private/BlockCorrelator.cpp b/Engine/Source/Runtime/SignalProcessing/Private/BlockCorrelator.cpp
--- a/Engine/Source/Runtime/SignalProcessing/Private/BlockCorrelator.cpp
+++ b/Engine/Source/Runtime/SignalProcessing/Private/BlockCorrelator.cpp
@@ -75,7 +75,7 @@ namespace Audio
 		return NumValuesInFFTRealBuffer;
 	}
 
-	void FBlockCorrelator::CrossCorrelate(const AlignedFloatBuffer& InputA, const AlignedFloatBuffer& InputB, AlignedFloatBuffer& Output)
+	void FBlockCorrelator::CrossCorrelate(const FAlignedFloatBuffer& InputA, const FAlignedFloatBuffer& InputB, FAlignedFloatBuffer& Output)
 	{
 		check(InputA.Num() == NumValuesInBlock);
 		check(InputB.Num() == NumValuesInBlock);
@@ -104,7 +104,7 @@ namespace Audio
 		}
 	}
 
-	void FBlockCorrelator::AutoCorrelate(const AlignedFloatBuffer& Input, AlignedFloatBuffer& Output)
+	void FBlockCorrelator::AutoCorrelate(const FAlignedFloatBuffer& Input, FAlignedFloatBuffer& Output)
 	{
 		check(Input.Num() == NumValuesInBlock);
 		check(Output.Num() == NumValuesInFFTRealBuffer);
@@ -129,7 +129,7 @@ namespace Audio
 		}
 	}
 
-	void FBlockCorrelator::CyclicCrossCorrelate(const AlignedFloatBuffer& InputA, const AlignedFloatBuffer& InputB, AlignedFloatBuffer& Output)
+	void FBlockCorrelator::CyclicCrossCorrelate(const FAlignedFloatBuffer& InputA, const FAlignedFloatBuffer& InputB, FAlignedFloatBuffer& Output)
 	{
 		check(InputA.Num() == NumValuesInFFTRealBuffer);
 		check(InputB.Num() == NumValuesInFFTRealBuffer);
@@ -151,7 +151,7 @@ namespace Audio
 		}
 	}
 
-	void FBlockCorrelator::CyclicAutoCorrelate(const AlignedFloatBuffer& Input, AlignedFloatBuffer& Output)
+	void FBlockCorrelator::CyclicAutoCorrelate(const FAlignedFloatBuffer& Input, FAlignedFloatBuffer& Output)
 	{
 		check(Input.Num() == NumValuesInFFTRealBuffer);
 		check(Output.Num() == NumValuesInFFTRealBuffer);
@@ -195,7 +195,7 @@ namespace Audio
 			Window.ApplyToBuffer(NormalizationBuffer.GetData());
 		}
 
-		AlignedFloatBuffer CorrelationBuffer;
+		FAlignedFloatBuffer CorrelationBuffer;
 		CorrelationBuffer.AddUninitialized(NumValuesInFFTRealBuffer);
 
 		// Take auto correlation of the window.
diff --git a/Engine/Source/Runtime/SignalProcessing/Private/BufferOnePoleLPF.cpp b/Engine/Source/Runtime/SignalProcessing/Private/BufferOnePoleLPF.cpp
--- a/Engine/Source/Runtime/SignalProcessing/Private/BufferOnePoleLPF.cpp
+++ b/Engine/Source/Runtime/SignalProcessing/Private/BufferOnePoleLPF.cpp
@@ -44,7 +44,7 @@ void FBufferOnePoleLPF::SetFrequency(const float InFrequency)
 	}
 }
 
-void FBufferOnePoleLPF::ProcessAudio(const Audio::AlignedFloatBuffer& InSamples, Audio::AlignedFloatBuffer& OutSamples)
+void FBufferOnePoleLPF::ProcessAudio(const Audio::FAlignedFloatBuffer& InSamples, Audio::FAlignedFloatBuffer& OutSamples)
 {
 	int32 Index;
 	const int32 InNum = InSamples.Num();
diff --git a/Engine/Source/Runtime/SignalProcessing/Private/BufferVectorOperations.cpp b/Engine/Source/Runtime/SignalProcessing/Private/BufferVectorOperations.cpp
--- a/Engine/Source/Runtime/SignalProcessing/Private/BufferVectorOperations.cpp
+++ b/Engine/Source/Runtime/SignalProcessing/Private/BufferVectorOperations.cpp
@@ -42,7 +42,7 @@ namespace Audio
 	}
 
 	/* Clamps values in the buffer to be between InMinValue and InMaxValue */
-	void BufferRangeClampFast(AlignedFloatBuffer& InOutBuffer, float InMinValue, float InMaxValue)
+	void BufferRangeClampFast(FAlignedFloatBuffer& InOutBuffer, float InMinValue, float InMaxValue)
 	{
 		return BufferRangeClampFast(InOutBuffer.GetData(), InOutBuffer.Num(), InMinValue, InMaxValue);
 	}
@@ -126,7 +126,7 @@ namespace Audio
 #endif
 	}
 
-	void MultiplyBufferByConstantInPlace(AlignedFloatBuffer& InBuffer, float InGain)
+	void MultiplyBufferByConstantInPlace(FAlignedFloatBuffer& InBuffer, float InGain)
 	{
 		MultiplyBufferByConstantInPlace(InBuffer.GetData(), InBuffer.Num(), InGain);
 	}
@@ -144,7 +144,7 @@ namespace Audio
 	}
 
 	// Adds a constant to a buffer (useful for DC offset removal)
-	void AddConstantToBufferInplace(AlignedFloatBuffer& InBuffer, float InConstant)
+	void AddConstantToBufferInplace(FAlignedFloatBuffer& InBuffer, float InConstant)
 	{
 		AddConstantToBufferInplace(InBuffer.GetData(), InBuffer.Num(), InConstant);
 	}
@@ -161,7 +161,7 @@ namespace Audio
 		}
 	}
 
-	void BufferSetToConstantInplace(AlignedFloatBuffer& InBuffer, float InConstant)
+	void BufferSetToConstantInplace(FAlignedFloatBuffer& InBuffer, float InConstant)
 	{
 		BufferSetToConstantInplace(InBuffer.GetData(), InBuffer.Num(), InConstant);
 	}
@@ -257,7 +257,7 @@ namespace Audio
 		}
 	}
 
-	void FadeBufferFast(AlignedFloatBuffer& OutFloatBuffer, const float StartValue, const float EndValue)
+	void FadeBufferFast(FAlignedFloatBuffer& OutFloatBuffer, const float StartValue, const float EndValue)
 	{
 		FadeBufferFast(OutFloatBuffer.GetData(), OutFloatBuffer.Num(), StartValue, EndValue);
 	}
@@ -516,7 +516,7 @@ namespace Audio
 	}
 	
 	/* Performs element-wise in-place subtraction placing the result in the minuend. InOutMinuend = InOutMinuend - InSubtrahend */
-	void BufferSubtractInPlace2Fast(AlignedFloatBuffer& InOutMinuend, const FAlignedFloatBuffer& InSubtrahend)
+	void BufferSubtractInPlace2Fast(FAlignedFloatBuffer& InOutMinuend, const FAlignedFloatBuffer& InSubtrahend)
 	{
 		checkf(InOutMinuend.Num() == InSubtrahend.Num(), TEXT("Input buffers must be equal length"));
 		BufferSubtractInPlace2Fast(InOutMinuend.GetData(), InSubtrahend.GetData(), InOutMinuend.Num());
@@ -764,7 +764,7 @@ namespace Audio
 	 * For clarity, the layout of vectors for each function variant is given in a block comment above that function.
 	 */
 
-	void Apply2ChannelGain(AlignedFloatBuffer& StereoBuffer, const float* RESTRICT Gains)
+	void Apply2ChannelGain(FAlignedFloatBuffer& StereoBuffer, const float* RESTRICT Gains)
 	{
 		Apply2ChannelGain(StereoBuffer.GetData(), StereoBuffer.Num(), Gains);
 	}
@@ -781,7 +781,7 @@ namespace Audio
 		}
 	}
 
-	void Apply2ChannelGain(AlignedFloatBuffer& StereoBuffer, const float* RESTRICT StartGains, const float* RESTRICT EndGains)
+	void Apply2ChannelGain(FAlignedFloatBuffer& StereoBuffer, const float* RESTRICT StartGains, const float* RESTRICT EndGains)
 	{
 		Apply2ChannelGain(StereoBuffer.GetData(), StereoBuffer.Num(), StartGains, EndGains);
 	}
@@ -1009,7 +1009,7 @@ namespace Audio
 		}
 	}
 
-	void Apply4ChannelGain(AlignedFloatBuffer& InterleavedBuffer, const float* RESTRICT Gains)
+	void Apply4ChannelGain(FAlignedFloatBuffer& InterleavedBuffer, const float* RESTRICT Gains)
 	{
 		Apply4ChannelGain(InterleavedBuffer.GetData(), InterleavedBuffer.Num(), Gains);
 	}
@@ -1026,7 +1026,7 @@ namespace Audio
 		}
 	}
 
-	void Apply4ChannelGain(AlignedFloatBuffer& InterleavedBuffer, const float* RESTRICT StartGains, const float* RESTRICT EndGains)
+	void Apply4ChannelGain(FAlignedFloatBuffer& InterleavedBuffer, const float* RESTRICT StartGains, const float* RESTRICT EndGains)
 	{
 		Apply4ChannelGain(InterleavedBuffer.GetData(), InterleavedBuffer.Num(), StartGains, EndGains);
 	}
@@ -1203,7 +1203,7 @@ namespace Audio
 		}
 	}
 
-	void Apply6ChannelGain(AlignedFloatBuffer& InterleavedBuffer, const float* RESTRICT Gains)
+	void Apply6ChannelGain(FAlignedFloatBuffer& InterleavedBuffer, const float* RESTRICT Gains)
 	{
 		Apply6ChannelGain(InterleavedBuffer.GetData(), InterleavedBuffer.Num(), Gains);
 	}
@@ -1230,7 +1230,7 @@ namespace Audio
 		}
 	}
 
-	void Apply6ChannelGain(AlignedFloatBuffer& InterleavedBuffer, const float* RESTRICT StartGains, const float* RESTRICT EndGains)
+	void Apply6ChannelGain(FAlignedFloatBuffer& InterleavedBuffer, const float* RESTRICT StartGains, const float* RESTRICT EndGains)
 	{
 		Apply6ChannelGain(InterleavedBuffer.GetData(), InterleavedBuffer.Num(), StartGains, EndGains);
 	}
@@ -1542,7 +1542,7 @@ namespace Audio
 		}
 	}
 
-	void Apply8ChannelGain(AlignedFloatBuffer& InterleavedBuffer, const float* RESTRICT Gains)
+	void Apply8ChannelGain(FAlignedFloatBuffer& InterleavedBuffer, const float* RESTRICT Gains)
 	{
 		Apply8ChannelGain(InterleavedBuffer.GetData(), InterleavedBuffer.Num(), Gains);
 	}
@@ -1564,7 +1564,7 @@ namespace Audio
 		}
 	}
 
-	void Apply8ChannelGain(AlignedFloatBuffer& InterleavedBuffer, const float* RESTRICT StartGains, const float* RESTRICT EndGains)
+	void Apply8ChannelGain(FAlignedFloatBuffer& InterleavedBuffer, const float* RESTRICT StartGains, const float* RESTRICT EndGains)
 	{
 		Apply8ChannelGain(InterleavedBuffer.GetData(), InterleavedBuffer.Num(), StartGains, EndGains);
 	}
@@ -2083,7 +2083,7 @@ namespace Audio
 		MultiplyBufferByConstantInPlace(DeltaBuffer, OneOverLerpLength);
 	}
 
-	bool FBufferLinearEase::Update(AlignedFloatBuffer & InSourceValues)
+	bool FBufferLinearEase::Update(FAlignedFloatBuffer & InSourceValues)
 	{
 		check(InSourceValues.Num() == BufferLength);
 		check(CurrentLerpStep != LerpLength);
diff --git a/Engine/Source/Runtime/SignalProcessing/Private/ConstantQ.cpp b/Engine/Source/Runtime/SignalProcessing/Private/ConstantQ.cpp
--- a/Engine/Source/Runtime/SignalProcessing/Private/ConstantQ.cpp
+++ b/Engine/Source/Runtime/SignalProcessing/Private/ConstantQ.cpp
@@ -8,7 +8,7 @@ namespace Audio
 {
 	namespace PseudoConstantQIntrinsics
 	{
-		void FillArrayWithTruncatedGaussian(const float InCenterFreq, const float InBandWidth, const int32 InFFTSize, const float InSampleRate, AlignedFloatBuffer& OutOffsetArray, int32& OutOffsetIndex)
+		void FillArrayWithTruncatedGaussian(const float InCenterFreq, const float InBandWidth, const int32 InFFTSize, const float InSampleRate, FAlignedFloatBuffer& OutOffsetArray, int32& OutOffsetIndex)
 		{
 			check(InBandWidth > 0.f);
 			check(InFFTSize > 0);
@@ -68,7 +68,7 @@ namespace Audio
 	}
 
 
-	void FPseudoConstantQ::FillArrayWithConstantQBand(const FPseudoConstantQBandSettings& InSettings, AlignedFloatBuffer& OutOffsetArray, int32& OutOffsetIndex)
+	void FPseudoConstantQ::FillArrayWithConstantQBand(const FPseudoConstantQBandSettings& InSettings, FAlignedFloatBuffer& OutOffsetArray, int32& OutOffsetIndex)
 	{
 		check(InSettings.SampleRate > 0.f);
 		check(InSettings.FFTSize > 0);
@@ -187,7 +187,7 @@ namespace Audio
 			}
 
 			// Create gaussian centered around center freq and with appropriate bandwidth
-			AlignedFloatBuffer OffsetBandWeights;
+			FAlignedFloatBuffer OffsetBandWeights;
 			int32 OffsetBandWeightsIndex = 0;
 
 			FPseudoConstantQ::FillArrayWithConstantQBand(BandSettings, OffsetBandWeights, OffsetBandWeightsIndex);
diff --git a/Engine/Source/Runtime/SignalProcessing/Private/DynamicDelayAPF.cpp b/Engine/Source/Runtime/SignalProcessing/Private/DynamicDelayAPF.cpp
--- a/Engine/Source/Runtime/SignalProcessing/Private/DynamicDelayAPF.cpp
+++ b/Engine/Source/Runtime/SignalProcessing/Private/DynamicDelayAPF.cpp
@@ -49,7 +49,7 @@ FDynamicDelayAPF::FDynamicDelayAPF(float InG, int32 InMinDelay, int32 InMaxDelay
 FDynamicDelayAPF::~FDynamicDelayAPF()
 {}
 
-void FDynamicDelayAPF::ProcessAudio(const AlignedFloatBuffer& InSamples, const AlignedFloatBuffer& InSampleDelays, AlignedFloatBuffer& OutSamples)
+void FDynamicDelayAPF::ProcessAudio(const FAlignedFloatBuffer& InSamples, const FAlignedFloatBuffer& InSampleDelays, FAlignedFloatBuffer& OutSamples)
 {
 	const int32 InNum = InSamples.Num();
 	checkf(InNum == InSampleDelays.Num(), TEXT("InSamples [%d], InSampleDelays [%d] length mismatch"), InNum, InSampleDelays.Num());
@@ -91,7 +91,7 @@ void FDynamicDelayAPF::ProcessAudio(const AlignedFloatBuffer& InSamples, const A
 	}
 }
 
-void FDynamicDelayAPF::ProcessAudioBlock(const float* InSamples, const AlignedFloatBuffer& InFractionalDelays, const int32 InNum, float* OutSamples)
+void FDynamicDelayAPF::ProcessAudioBlock(const float* InSamples, const FAlignedFloatBuffer& InFractionalDelays, const int32 InNum, float* OutSamples)
 {
 	// Make a copy of the delay line w[n - d_int]
 	WorkBufferA.Reset(InNum);
diff --git a/Engine/Source/Runtime/SignalProcessing/Private/EarlyReflectionsFast.cpp b/Engine/Source/Runtime/SignalProcessing/Private/EarlyReflectionsFast.cpp
--- a/Engine/Source/Runtime/SignalProcessing/Private/EarlyReflectionsFast.cpp
+++ b/Engine/Source/Runtime/SignalProcessing/Private/EarlyReflectionsFast.cpp
@@ -154,7 +154,7 @@ void FEarlyReflectionsFast::ApplySettings()
 	RightFDN.SetCoefficients(RightCoefficients);
 }
 
-void FEarlyReflectionsFast::ProcessAudio(const AlignedFloatBuffer& InSamples, const int32 InNumChannels, AlignedFloatBuffer& OutLeftSamples, AlignedFloatBuffer& OutRightSamples)
+void FEarlyReflectionsFast::ProcessAudio(const FAlignedFloatBuffer& InSamples, const int32 InNumChannels, FAlignedFloatBuffer& OutLeftSamples, FAlignedFloatBuffer& OutRightSamples)
 {
 	checkf((InNumChannels == 1) || (InNumChannels == 2), TEXT("EarlyReflections only supports one or two channel input samples."));
 
diff --git a/Engine/Source/Runtime/SignalProcessing/Private/FeedbackDelayNetwork.cpp b/Engine/Source/Runtime/SignalProcessing/Private/FeedbackDelayNetwork.cpp
--- a/Engine/Source/Runtime/SignalProcessing/Private/FeedbackDelayNetwork.cpp
+++ b/Engine/Source/Runtime/SignalProcessing/Private/FeedbackDelayNetwork.cpp
@@ -88,7 +88,7 @@ void FFeedbackDelayNetwork::SetCoefficients(const FFDNCoefficients& InCoefficien
 	Coefficients = InCoefficients;
 }
 
-void FFeedbackDelayNetwork::ProcessAudio(const AlignedFloatBuffer& InSamples, AlignedFloatBuffer& OutSamples)
+void FFeedbackDelayNetwork::ProcessAudio(const FAlignedFloatBuffer& InSamples, FAlignedFloatBuffer& OutSamples)
 {
 	const int32 InNum = InSamples.Num();
 	const float* InSampleData = InSamples.GetData();
diff --git a/Engine/Source/Runtime/SignalProcessing/Private/FloatArrayMath.cpp b/Engine/Source/Runtime/SignalProcessing/Private/FloatArrayMath.cpp
--- a/Engine/Source/Runtime/SignalProcessing/Private/FloatArrayMath.cpp
+++ b/Engine/Source/Runtime/SignalProcessing/Private/FloatArrayMath.cpp
@@ -25,7 +25,7 @@ namespace Audio
 		}
 	}
 
-	void ArraySum(const AlignedFloatBuffer& InValues, float& OutSum)
+	void ArraySum(const FAlignedFloatBuffer& InValues, float& OutSum)
 	{
 		OutSum = 0.f;
 
@@ -366,7 +366,7 @@ namespace Audio
 		}
 	}
 
-	void ArrayMultiplyInPlace(const AlignedFloatBuffer& InValues1, AlignedFloatBuffer& InValues2)
+	void ArrayMultiplyInPlace(const FAlignedFloatBuffer& InValues1, FAlignedFloatBuffer& InValues2)
 	{
 		check(InValues1.Num() == InValues2.Num());
 
@@ -412,7 +412,7 @@ namespace Audio
 		}
 	}
 
-	void ArrayComplexMultiplyInPlace(const AlignedFloatBuffer& InValues1, AlignedFloatBuffer& InValues2)
+	void ArrayComplexMultiplyInPlace(const FAlignedFloatBuffer& InValues1, FAlignedFloatBuffer& InValues2)
 	{
 		check(InValues1.Num() == InValues2.Num());
 
@@ -464,7 +464,7 @@ namespace Audio
 		}
 	}
 
-	void ArrayMultiplyByConstantInPlace(AlignedFloatBuffer& InValues, float InMultiplier)
+	void ArrayMultiplyByConstantInPlace(FAlignedFloatBuffer& InValues, float InMultiplier)
 	{
 		const int32 Num = InValues.Num();
 		const int32 NumToSimd = Num & MathIntrinsics::SimdMask;
@@ -500,7 +500,7 @@ namespace Audio
 		}
 	}
 
-	void ArrayAddInPlace(const AlignedFloatBuffer& InValues, AlignedFloatBuffer& InAccumulateValues)
+	void ArrayAddInPlace(const FAlignedFloatBuffer& InValues, FAlignedFloatBuffer& InAccumulateValues)
 	{
 		check(InValues.Num() == InAccumulateValues.Num());
 
@@ -544,7 +544,7 @@ namespace Audio
 		}
 	}
 
-	void ArrayMultiplyAddInPlace(const AlignedFloatBuffer& InValues, float InMultiplier, AlignedFloatBuffer& InAccumulateValues)
+	void ArrayMultiplyAddInPlace(const FAlignedFloatBuffer& InValues, float InMultiplier, FAlignedFloatBuffer& InAccumulateValues)
 	{
 		check(InValues.Num() == InAccumulateValues.Num());
 
@@ -585,7 +585,7 @@ namespace Audio
 		}
 	}
 
-	void ArrayLerpAddInPlace(const AlignedFloatBuffer& InValues, float InStartMultiplier, float InEndMultiplier, AlignedFloatBuffer& InAccumulateValues)
+	void ArrayLerpAddInPlace(const FAlignedFloatBuffer& InValues, float InStartMultiplier, float InEndMultiplier, FAlignedFloatBuffer& InAccumulateValues)
 	{
 		check(InValues.Num() == InAccumulateValues.Num());
 
@@ -632,7 +632,7 @@ namespace Audio
 		}
 	}
 
-	void ArraySubtractByConstantInPlace(AlignedFloatBuffer& InValues, float InSubtrahend)
+	void ArraySubtractByConstantInPlace(FAlignedFloatBuffer& InValues, float InSubtrahend)
 	{
 		const int32 Num = InValues.Num();
 		const int32 NumToSimd = Num & MathIntrinsics::SimdMask;
@@ -709,7 +709,7 @@ namespace Audio
 		}
 	}
 
-	void ArrayComplexConjugate(const AlignedFloatBuffer& InValues, AlignedFloatBuffer& OutValues)
+	void ArrayComplexConjugate(const FAlignedFloatBuffer& InValues, FAlignedFloatBuffer& OutValues)
 	{
 		check(OutValues.Num() == InValues.Num());
 
@@ -754,7 +754,7 @@ namespace Audio
 		}
 	}
 
-	void ArrayComplexConjugateInPlace(AlignedFloatBuffer& InValues)
+	void ArrayComplexConjugateInPlace(FAlignedFloatBuffer& InValues)
 	{
 		const int32 Num = InValues.Num();
 		const int32 NumToSimd = Num & MathIntrinsics::SimdMask;
@@ -795,7 +795,7 @@ namespace Audio
 		}
 	}
 
-	void ArrayMagnitudeToDecibelInPlace(AlignedFloatBuffer& InValues, float InMinimumDb)
+	void ArrayMagnitudeToDecibelInPlace(FAlignedFloatBuffer& InValues, float InMinimumDb)
 	{
 		const int32 Num = InValues.Num();
 		const int32 NumToSimd = Num & MathIntrinsics::SimdMask;
@@ -841,7 +841,7 @@ namespace Audio
 		}
 	}
 
-	void ArrayPowerToDecibelInPlace(AlignedFloatBuffer& InValues, float InMinimumDb)
+	void ArrayPowerToDecibelInPlace(FAlignedFloatBuffer& InValues, float InMinimumDb)
 	{
 		const int32 Num = InValues.Num();
 		const int32 NumToSimd = Num & MathIntrinsics::SimdMask;
@@ -894,7 +894,7 @@ namespace Audio
 		}
 	}
 
-	void ArrayComplexToPower(const AlignedFloatBuffer& InComplexValues, AlignedFloatBuffer& OutPowerValues)
+	void ArrayComplexToPower(const FAlignedFloatBuffer& InComplexValues, FAlignedFloatBuffer& OutPowerValues)
 	{
 		check((InComplexValues.Num() % 2) == 0);
 		check(InComplexValues.Num() == (OutPowerValues.Num() * 2));
@@ -983,7 +983,7 @@ namespace Audio
 		TransformArray(InView.GetData(), OutArray.GetData());
 	}
 
-	void FContiguousSparse2DKernelTransform::TransformArray(TArrayView<const float> InView, AlignedFloatBuffer& OutArray) const
+	void FContiguousSparse2DKernelTransform::TransformArray(TArrayView<const float> InView, FAlignedFloatBuffer& OutArray) const
 	{	
 		check(InView.Num() == NumIn);
 
diff --git a/Engine/Source/Runtime/SignalProcessing/Private/IntegerDelay.cpp b/Engine/Source/Runtime/SignalProcessing/Private/IntegerDelay.cpp
--- a/Engine/Source/Runtime/SignalProcessing/Private/IntegerDelay.cpp
+++ b/Engine/Source/Runtime/SignalProcessing/Private/IntegerDelay.cpp
@@ -65,7 +65,7 @@ int32 FIntegerDelay::GetNumDelaySamples() const
 	return NumDelaySamples;
 }
 
-void FIntegerDelay::ProcessAudio(const Audio::AlignedFloatBuffer& InSamples, Audio::AlignedFloatBuffer& OutSamples)
+void FIntegerDelay::ProcessAudio(const Audio::FAlignedFloatBuffer& InSamples, Audio::FAlignedFloatBuffer& OutSamples)
 {
 	const float* InSampleData = InSamples.GetData();
 	const int32 InNum = InSamples.Num();
@@ -102,7 +102,7 @@ void FIntegerDelay::ProcessAudioBlock(const float* InSamples, const int32 InNum,
 	DelayLine->RemoveSamples(InNum);
 }
 
-void FIntegerDelay::PeekDelayLine(int32 InNum, Audio::AlignedFloatBuffer& OutSamples)
+void FIntegerDelay::PeekDelayLine(int32 InNum, Audio::FAlignedFloatBuffer& OutSamples)
 {
 	int32 NumToInspect = FMath::Min(DelayLine->GetNumAvailable(), NumDelaySamples);
 	NumToInspect = FMath::Min(InNum, NumToInspect);
diff --git a/Engine/Source/Runtime/SignalProcessing/Private/LateReflectionsFast.cpp b/Engine/Source/Runtime/SignalProcessing/Private/LateReflectionsFast.cpp
--- a/Engine/Source/Runtime/SignalProcessing/Private/LateReflectionsFast.cpp
+++ b/Engine/Source/Runtime/SignalProcessing/Private/LateReflectionsFast.cpp
@@ -155,9 +155,9 @@ FLateReflectionsPlate::~FLateReflectionsPlate()
 
 
 void FLateReflectionsPlate::ProcessAudioFrames(
-		const AlignedFloatBuffer& InSamples,
-		const AlignedFloatBuffer& InFeedbackSamples,
-		const AlignedFloatBuffer& InDelayModulations,
+		const FAlignedFloatBuffer& InSamples,
+		const FAlignedFloatBuffer& InFeedbackSamples,
+		const FAlignedFloatBuffer& InDelayModulations,
 		FLateReflectionsPlateOutputs& OutPlateSamples)
 {
 	
@@ -261,7 +261,7 @@ int32 FLateReflectionsPlate::GetNumInternalBufferSamples() const
 	return NumInternalBufferSamples;
 }
 
-void FLateReflectionsPlate::PeekDelayLine(int32 InNum, AlignedFloatBuffer& OutSamples)
+void FLateReflectionsPlate::PeekDelayLine(int32 InNum, FAlignedFloatBuffer& OutSamples)
 {
 	DelayI->PeekDelayLine(InNum, OutSamples);
 }
@@ -366,7 +366,7 @@ void FLateReflectionsFast::SetSettings(const FLateReflectionsFastSettings& InSet
 	ApplySettings();
 }
 
-void FLateReflectionsFast::ProcessAudio(const AlignedFloatBuffer& InSamples, const int32 InNumChannels, AlignedFloatBuffer& OutLeftSamples, AlignedFloatBuffer& OutRightSamples)
+void FLateReflectionsFast::ProcessAudio(const FAlignedFloatBuffer& InSamples, const int32 InNumChannels, FAlignedFloatBuffer& OutLeftSamples, FAlignedFloatBuffer& OutRightSamples)
 {
 	checkf((InNumChannels == 1) || (InNumChannels == 2), TEXT("FLateReflections only supports 1 or 2 channel input audio"));
 
@@ -532,7 +532,7 @@ void FLateReflectionsFast::ApplySettings()
 }
 
 
-void FLateReflectionsFast::GeneraterPlateModulations(const int32 InNum, AlignedFloatBuffer& OutLeftDelays, AlignedFloatBuffer& OutRightDelays)
+void FLateReflectionsFast::GeneraterPlateModulations(const int32 InNum, FAlignedFloatBuffer& OutLeftDelays, FAlignedFloatBuffer& OutRightDelays)
 {
 	OutLeftDelays.Reset(InNum);
 	OutRightDelays.Reset(InNum);
diff --git a/Engine/Source/Runtime/SignalProcessing/Private/LinearInterpFractionalDelay.cpp b/Engine/Source/Runtime/SignalProcessing/Private/LinearInterpFractionalDelay.cpp
--- a/Engine/Source/Runtime/SignalProcessing/Private/LinearInterpFractionalDelay.cpp
+++ b/Engine/Source/Runtime/SignalProcessing/Private/LinearInterpFractionalDelay.cpp
@@ -60,7 +60,7 @@ void FLinearInterpFractionalDelay::Reset()
 }
 
 
-void FLinearInterpFractionalDelay::ProcessAudio(const AlignedFloatBuffer& InSamples, const AlignedFloatBuffer& InDelays, AlignedFloatBuffer& OutSamples)
+void FLinearInterpFractionalDelay::ProcessAudio(const FAlignedFloatBuffer& InSamples, const FAlignedFloatBuffer& InDelays, FAlignedFloatBuffer& OutSamples)
 {
 	const int32 InNum = InSamples.Num();
 	checkf(InNum == InDelays.Num(), TEXT("Input buffers must be equal length"));
diff --git a/Engine/Source/Runtime/SignalProcessing/Private/LongDelayAPF.cpp b/Engine/Source/Runtime/SignalProcessing/Private/LongDelayAPF.cpp
--- a/Engine/Source/Runtime/SignalProcessing/Private/LongDelayAPF.cpp
+++ b/Engine/Source/Runtime/SignalProcessing/Private/LongDelayAPF.cpp
@@ -34,7 +34,7 @@ FLongDelayAPF::FLongDelayAPF(float InG, int32 InNumDelaySamples, int32 InMaxNumI
 FLongDelayAPF::~FLongDelayAPF()
 {}
 
-void FLongDelayAPF::ProcessAudio(const AlignedFloatBuffer& InSamples, AlignedFloatBuffer& OutSamples)
+void FLongDelayAPF::ProcessAudio(const FAlignedFloatBuffer& InSamples, FAlignedFloatBuffer& OutSamples)
 {
 	const float* InData = InSamples.GetData();
 	const int32 InNum = InSamples.Num();
@@ -66,7 +66,7 @@ void FLongDelayAPF::ProcessAudio(const AlignedFloatBuffer& InSamples, AlignedFlo
 	}
 }
 
-void FLongDelayAPF::ProcessAudio(const AlignedFloatBuffer& InSamples, AlignedFloatBuffer& OutSamples, AlignedFloatBuffer& OutDelaySamples)
+void FLongDelayAPF::ProcessAudio(const FAlignedFloatBuffer& InSamples, FAlignedFloatBuffer& OutSamples, FAlignedFloatBuffer& OutDelaySamples)
 {
 	const float* InData = InSamples.GetData();
 	const int32 InNum = InSamples.Num();
diff --git a/Engine/Source/Runtime/SignalProcessing/Private/ReverbFast.cpp b/Engine/Source/Runtime/SignalProcessing/Private/ReverbFast.cpp
--- a/Engine/Source/Runtime/SignalProcessing/Private/ReverbFast.cpp
+++ b/Engine/Source/Runtime/SignalProcessing/Private/ReverbFast.cpp
@@ -67,7 +67,7 @@ namespace Audio {
 		return Settings;
 	}
 	// Process a buffer of input audio samples.
-	void FPlateReverbFast::ProcessAudio(const AlignedFloatBuffer& InSamples, const int32 InNumChannels, AlignedFloatBuffer& OutSamples, const int32 OutNumChannels)
+	void FPlateReverbFast::ProcessAudio(const FAlignedFloatBuffer& InSamples, const int32 InNumChannels, FAlignedFloatBuffer& OutSamples, const int32 OutNumChannels)
 	{
 		if(InSamples.Num() == 0)
 		{
@@ -157,7 +157,7 @@ namespace Audio {
 
 	// Copy reverberated samples to interleaved output samples. Map channels according to internal settings.
 	// InFrontLeftSamples and InFrontRightSamples may be modified in-place.
-	void FPlateReverbFast::InterleaveAndMixOutput(const AlignedFloatBuffer& InFrontLeftSamples, const AlignedFloatBuffer& InFrontRightSamples, AlignedFloatBuffer& OutSamples, const int32 OutNumChannels)
+	void FPlateReverbFast::InterleaveAndMixOutput(const FAlignedFloatBuffer& InFrontLeftSamples, const FAlignedFloatBuffer& InFrontRightSamples, FAlignedFloatBuffer& OutSamples, const int32 OutNumChannels)
 	{
 		check(InFrontLeftSamples.Num() == InFrontRightSamples.Num())
 
diff --git a/Engine/Source/Runtime/SignalProcessing/Private/SpectrumAnalyzer.cpp b/Engine/Source/Runtime/SignalProcessing/Private/SpectrumAnalyzer.cpp
--- a/Engine/Source/Runtime/SignalProcessing/Private/SpectrumAnalyzer.cpp
+++ b/Engine/Source/Runtime/SignalProcessing/Private/SpectrumAnalyzer.cpp
@@ -232,10 +232,10 @@ namespace Audio
 				int32 StartIndex;
 
 				// Weights (offset by start index) to apply to power spectrum
-				AlignedFloatBuffer Weights;
+				FAlignedFloatBuffer Weights;
 
 				// Internal buffer used when calculating band.
-				mutable AlignedFloatBuffer WorkBuffer;
+				mutable FAlignedFloatBuffer WorkBuffer;
 
 				virtual void Update(const FSpectrumBandExtractorSettings& InSettings, const FSpectrumBandExtractorSpectrumSettings& InSpectrumSettings) override
 				{
@@ -466,7 +466,7 @@ namespace Audio
 
 
 			// Extract band from input.
-			virtual void ExtractBands(const AlignedFloatBuffer& InComplexBuffer, double InTimestamp, TArray<float>& OutValues) override
+			virtual void ExtractBands(const FAlignedFloatBuffer& InComplexBuffer, double InTimestamp, TArray<float>& OutValues) override
 			{
 				const int32 NumComplex = InComplexBuffer.Num();
 
@@ -549,7 +549,7 @@ namespace Audio
 			}
 
 			template<typename T>
-			void ExtractBands(const AlignedFloatBuffer& InPowerSpectrum, const TArray<T>& InBandSpecs, TArray<float>& OutValues) const
+			void ExtractBands(const FAlignedFloatBuffer& InPowerSpectrum, const TArray<T>& InBandSpecs, TArray<float>& OutValues) const
 			{
 				float* OutData = OutValues.GetData();
 				int32 OutNum = OutValues.Num();
@@ -609,7 +609,7 @@ namespace Audio
 			FSpectrumBandExtractorSpectrumSettings SpectrumSettings;
 			double LastTimestamp;
 
-			AlignedFloatBuffer PowerSpectrum;
+			FAlignedFloatBuffer PowerSpectrum;
 
 			TArray<FNNBandSpec> NNBandSpecs;
 			TArray<FLerpBandSpec> LerpBandSpecs;
@@ -762,7 +762,7 @@ namespace Audio
 		bSettingsWereUpdated = false;
 	}
 
-	void FSpectrumAnalyzer::PerformInterpolation(const AlignedFloatBuffer& InComplexBuffer, FSpectrumAnalyzer::EPeakInterpolationMethod InMethod, const float InFreq, float& OutReal, float& OutImag)
+	void FSpectrumAnalyzer::PerformInterpolation(const FAlignedFloatBuffer& InComplexBuffer, FSpectrumAnalyzer::EPeakInterpolationMethod InMethod, const float InFreq, float& OutReal, float& OutImag)
 	{
 		const float* InComplexData = InComplexBuffer.GetData();
 		const int32 VectorLength = InComplexBuffer.Num();
@@ -857,7 +857,7 @@ namespace Audio
 			return 0.f;
 		}
 
-		const AlignedFloatBuffer* OutVector = nullptr;
+		const FAlignedFloatBuffer* OutVector = nullptr;
 		bool bShouldUnlockBuffer = true;
 
 		if (LockedFrequencyVector)
@@ -912,7 +912,7 @@ namespace Audio
 			return 0.f;
 		}
 
-		const AlignedFloatBuffer* OutVector = nullptr;
+		const FAlignedFloatBuffer* OutVector = nullptr;
 		bool bShouldUnlockBuffer = true;
 
 		if (LockedFrequencyVector)
@@ -958,7 +958,7 @@ namespace Audio
 			return;
 		}
 
-		const AlignedFloatBuffer* AnalysisBuffer = nullptr;
+		const FAlignedFloatBuffer* AnalysisBuffer = nullptr;
 		bool bShouldUnlockBuffer = true;
 
 		FSpectrumBandExtractorSpectrumSettings ExtractorSettings;
@@ -1051,7 +1051,7 @@ namespace Audio
 		}
 
 
-		AlignedFloatBuffer& FFTOutput = FrequencyBuffer.StartWorkOnBuffer();
+		FAlignedFloatBuffer& FFTOutput = FrequencyBuffer.StartWorkOnBuffer();
 
 		// If we have enough audio pushed to the spectrum analyzer and we have an available buffer to work in,
 		// we can start analyzing.
@@ -1132,7 +1132,7 @@ namespace Audio
 
 		for (int32 Index = 0; Index < SpectrumAnalyzerBufferSize; Index++)
 		{
-			AlignedFloatBuffer& Buffer = ComplexBuffers.Emplace_GetRef();
+			FAlignedFloatBuffer& Buffer = ComplexBuffers.Emplace_GetRef();
 
 			if (InNum > 0)
 			{
@@ -1173,7 +1173,7 @@ namespace Audio
 		check(InputIndex != OutputIndex);
 	}
 
-	AlignedFloatBuffer& FSpectrumAnalyzerBuffer::StartWorkOnBuffer()
+	FAlignedFloatBuffer& FSpectrumAnalyzerBuffer::StartWorkOnBuffer()
 	{
 		return ComplexBuffers[InputIndex];
 	}
@@ -1184,13 +1184,13 @@ namespace Audio
 		IncrementInputIndex();
 	}
 
-	const AlignedFloatBuffer& FSpectrumAnalyzerBuffer::LockMostRecentBuffer(double& OutTimestamp) const
+	const FAlignedFloatBuffer& FSpectrumAnalyzerBuffer::LockMostRecentBuffer(double& OutTimestamp) const
 	{
 		OutTimestamp = Timestamps[OutputIndex];
 		return ComplexBuffers[OutputIndex];
 	}
 
-	const AlignedFloatBuffer& FSpectrumAnalyzerBuffer::LockMostRecentBuffer() const
+	const FAlignedFloatBuffer& FSpectrumAnalyzerBuffer::LockMostRecentBuffer() const
 	{
 		return ComplexBuffers[OutputIndex];
 	}
diff --git a/Engine/Source/Runtime/SignalProcessing/Private/UniformPartitionConvolution.cpp b/Engine/Source/Runtime/SignalProcessing/Private/UniformPartitionConvolution.cpp
--- a/Engine/Source/Runtime/SignalProcessing/Private/UniformPartitionConvolution.cpp
+++ b/Engine/Source/Runtime/SignalProcessing/Private/UniformPartitionConvolution.cpp
@@ -249,7 +249,7 @@ namespace Audio
 	// to determine how many values can be SIMD'd.
 	const int32 FUniformPartitionConvolution::NumSimdMask = 0xFFFFFFFC;
 
-	void FUniformPartitionConvolution::VectorComplexMultiplyAdd(const AlignedFloatBuffer& InA, const AlignedFloatBuffer& InB, AlignedFloatBuffer& Out) const
+	void FUniformPartitionConvolution::VectorComplexMultiplyAdd(const FAlignedFloatBuffer& InA, const FAlignedFloatBuffer& InB, FAlignedFloatBuffer& Out) const
 	{
 		check(InA.Num() == InB.Num());
 		check(Out.Num() == InA.Num());
@@ -321,7 +321,7 @@ namespace Audio
 	}
 
 	// Multiply aligned buffer by constant gain.
-	void FUniformPartitionConvolution::VectorMultiplyByConstant(const AlignedFloatBuffer& InBuffer, float InConstant, AlignedFloatBuffer& OutBuffer) const
+	void FUniformPartitionConvolution::VectorMultiplyByConstant(const FAlignedFloatBuffer& InBuffer, float InConstant, FAlignedFloatBuffer& OutBuffer) const
 	{
 		check(InBuffer.Num() == OutBuffer.Num());
 
@@ -372,7 +372,7 @@ namespace Audio
 		FMemory::Memcpy(InputData, InSamples, BlockSize * sizeof(float));
 	}
 
-	const AlignedFloatBuffer& FUniformPartitionConvolution::FInput::GetTransformedBlock() const
+	const FAlignedFloatBuffer& FUniformPartitionConvolution::FInput::GetTransformedBlock() const
 	{
 		return OutputBuffer;
 	}
@@ -405,7 +405,7 @@ namespace Audio
 	{
 	}
 
-	AlignedFloatBuffer& FUniformPartitionConvolution::FOutput::GetTransformedBlock(int32 InBlockIndex)
+	FAlignedFloatBuffer& FUniformPartitionConvolution::FOutput::GetTransformedBlock(int32 InBlockIndex)
 	{
 		// Make block index relative to head index.
 		InBlockIndex = (HeadBlockIndex + InBlockIndex) % NumBlocks;
@@ -421,14 +421,14 @@ namespace Audio
 		HeadBlockIndex = (HeadBlockIndex + 1) % NumBlocks;
 
 		// Set tail block to zero
-		AlignedFloatBuffer& TailBuffer = Blocks[TailBlockIndex];
+		FAlignedFloatBuffer& TailBuffer = Blocks[TailBlockIndex];
 		FMemory::Memset(TailBuffer.GetData(), 0, NumFFTOutputFloats * sizeof(float));
 
 	}
 
 	void FUniformPartitionConvolution::FOutput::PopBlock(float* OutSamples)
 	{
-		const AlignedFloatBuffer& HeadBuffer = Blocks[HeadBlockIndex];
+		const FAlignedFloatBuffer& HeadBuffer = Blocks[HeadBlockIndex];
 		float* OutputData = OutputBuffer.GetData();
 
 		FFTAlgorithm->InverseComplexToReal(HeadBuffer.GetData(), OutputData);
@@ -520,7 +520,7 @@ namespace Audio
 		}
 	}
 
-	const AlignedFloatBuffer& FUniformPartitionConvolution::FImpulseResponse::GetTransformedBlock(int32 InBlockIndex) const
+	const FAlignedFloatBuffer& FUniformPartitionConvolution::FImpulseResponse::GetTransformedBlock(int32 InBlockIndex) const
 	{
 		return Blocks[InBlockIndex];
 	}
diff --git a/Engine/Source/Runtime/SignalProcessing/Private/UniformPartitionConvolution.h b/Engine/Source/Runtime/SignalProcessing/Private/UniformPartitionConvolution.h
--- a/Engine/Source/Runtime/SignalProcessing/Private/UniformPartitionConvolution.h
+++ b/Engine/Source/Runtime/SignalProcessing/Private/UniformPartitionConvolution.h
@@ -130,10 +130,10 @@ namespace Audio
 			};
 
 			// Vector optimized complex multiply add
-			void VectorComplexMultiplyAdd(const AlignedFloatBuffer& InA, const AlignedFloatBuffer& InB, AlignedFloatBuffer& Out) const;
+			void VectorComplexMultiplyAdd(const FAlignedFloatBuffer& InA, const FAlignedFloatBuffer& InB, FAlignedFloatBuffer& Out) const;
 
 			// Vector optimized multiply
-			void VectorMultiplyByConstant(const AlignedFloatBuffer& InBuffer, float InConstant, AlignedFloatBuffer& OutBuffer) const;
+			void VectorMultiplyByConstant(const FAlignedFloatBuffer& InBuffer, float InConstant, FAlignedFloatBuffer& OutBuffer) const;
 
 			// Contains a single input channel.
 			//
@@ -147,7 +147,7 @@ namespace Audio
 					void PushBlock(const float* InSamples);
 
 					// Reference to latest transformed block.
-					const AlignedFloatBuffer& GetTransformedBlock() const;
+					const FAlignedFloatBuffer& GetTransformedBlock() const;
 
 					// Clears the transformed block
 					void Reset();
@@ -158,8 +158,8 @@ namespace Audio
 				private: 
 
 					FSharedFFTRef FFTAlgorithm;
-					AlignedFloatBuffer InputBuffer;
-					AlignedFloatBuffer OutputBuffer;
+					FAlignedFloatBuffer InputBuffer;
+					FAlignedFloatBuffer OutputBuffer;
 			};
 
 			// Contains a single output channel
@@ -177,7 +177,7 @@ namespace Audio
 					~FOutput();
 
 					// Returns the history block at the given index. 
-					AlignedFloatBuffer& GetTransformedBlock(int32 InBlockIndex);	
+					FAlignedFloatBuffer& GetTransformedBlock(int32 InBlockIndex);	
 
 					// Pop the latest block without performing inverse FFT
 					void PopBlock();
@@ -195,10 +195,10 @@ namespace Audio
 				private:
 
 					FSharedFFTRef FFTAlgorithm;
-					AlignedFloatBuffer OutputBuffer;
+					FAlignedFloatBuffer OutputBuffer;
 
 					int32 HeadBlockIndex;
-					TArray<AlignedFloatBuffer> Blocks;
+					TArray<FAlignedFloatBuffer> Blocks;
 
 			};
 
@@ -219,7 +219,7 @@ namespace Audio
 					void SetImpulseResponse(const float* InSamples, int32 InNum);
 
 					// Returns a block of the transformed impulse response.
-					const AlignedFloatBuffer& GetTransformedBlock(int32 InBlockIndex) const;
+					const FAlignedFloatBuffer& GetTransformedBlock(int32 InBlockIndex) const;
 
 					// Returns the number of active blocks. This may be less than the maximum
 					// number of blocks if the impulse response is shorter than the maximum
@@ -237,8 +237,8 @@ namespace Audio
 					int32 NumFFTOutputFloats;		
 					int32 NumActiveBlocks;
 					int32 NumImpulseResponseSamples;
-					AlignedFloatBuffer FFTInput;
-					TArray<AlignedFloatBuffer> Blocks;
+					FAlignedFloatBuffer FFTInput;
+					TArray<FAlignedFloatBuffer> Blocks;
 			};
 
 			FUniformPartitionConvolutionSettings Settings;
@@ -246,7 +246,7 @@ namespace Audio
 			int32 BlockSize;
 			int32 NumFFTOutputFloats;
 			int32 NumBlocks;
-			AlignedFloatBuffer TransformedAndScaledInput;
+			FAlignedFloatBuffer TransformedAndScaledInput;
 			TSharedRef<IFFTAlgorithm> FFTAlgorithm;
 
 			TArray<FInput> Inputs;
diff --git a/Engine/Source/Runtime/SignalProcessing/Private/VectorFFT.cpp b/Engine/Source/Runtime/SignalProcessing/Private/VectorFFT.cpp
--- a/Engine/Source/Runtime/SignalProcessing/Private/VectorFFT.cpp
+++ b/Engine/Source/Runtime/SignalProcessing/Private/VectorFFT.cpp
@@ -1001,7 +1001,7 @@ namespace Audio
 			TArray<FFinalIndices> FinalIndices;
 			TArray<FFinalWeights> FinalWeights;
 
-			AlignedFloatBuffer InverseWorkBuffer;
+			FAlignedFloatBuffer InverseWorkBuffer;
 	};
 
 	// Maximum log 2 size of fft
diff --git a/Engine/Source/Runtime/SignalProcessing/Private/VectorFFT.h b/Engine/Source/Runtime/SignalProcessing/Private/VectorFFT.h
--- a/Engine/Source/Runtime/SignalProcessing/Private/VectorFFT.h
+++ b/Engine/Source/Runtime/SignalProcessing/Private/VectorFFT.h
@@ -80,10 +80,10 @@ namespace Audio
 			// Needed to convert buffers for implementation of real-to-complex fft using complex-to-comlex fft.
 			struct FConversionBuffers
 			{
-				AlignedFloatBuffer AlphaReal;
-				AlignedFloatBuffer AlphaImag;
-				AlignedFloatBuffer BetaReal;
-				AlignedFloatBuffer BetaImag;
+				FAlignedFloatBuffer AlphaReal;
+				FAlignedFloatBuffer AlphaImag;
+				FAlignedFloatBuffer BetaReal;
+				FAlignedFloatBuffer BetaImag;
 			};
 
 
@@ -96,7 +96,7 @@ namespace Audio
 			int32 FFTSize;
 			int32 Log2FFTSize;
 
-			AlignedFloatBuffer WorkBuffer;
+			FAlignedFloatBuffer WorkBuffer;
 
 			FConversionBuffers ForwardConvBuffers;
 			FConversionBuffers InverseConvBuffers;
diff --git a/Engine/Source/Runtime/SignalProcessing/Public/DSP/AllPassFractionalDelay.h b/Engine/Source/Runtime/SignalProcessing/Public/DSP/AllPassFractionalDelay.h
--- a/Engine/Source/Runtime/SignalProcessing/Public/DSP/AllPassFractionalDelay.h
+++ b/Engine/Source/Runtime/SignalProcessing/Public/DSP/AllPassFractionalDelay.h
@@ -20,7 +20,7 @@ namespace Audio
 
 		// Apply delay to InSamples. Fill OutSamples with data from the delay line at a delay of InDelay.
 		// InDelay must be equal length to InSamples. 
-		void ProcessAudio(const AlignedFloatBuffer& InSamples, const AlignedFloatBuffer& InDelays, AlignedFloatBuffer& OutSamples);
+		void ProcessAudio(const FAlignedFloatBuffer& InSamples, const FAlignedFloatBuffer& InDelays, FAlignedFloatBuffer& OutSamples);
 
 		// Set all values in internal delay line to zero. 
 		void Reset();
@@ -36,9 +36,9 @@ namespace Audio
 
 		TUniquePtr<FAlignedBlockBuffer> DelayLine;
 
-		AlignedFloatBuffer Coefficients;
-		AlignedFloatBuffer FractionalDelays;
-		AlignedInt32Buffer IntegerDelays;
-		AlignedInt32Buffer IntegerDelayOffsets;
+		FAlignedFloatBuffer Coefficients;
+		FAlignedFloatBuffer FractionalDelays;
+		FAlignedInt32Buffer IntegerDelays;
+		FAlignedInt32Buffer IntegerDelayOffsets;
 	};
 }
diff --git a/Engine/Source/Runtime/SignalProcessing/Public/DSP/AudioChannelFormatConverter.h b/Engine/Source/Runtime/SignalProcessing/Public/DSP/AudioChannelFormatConverter.h
--- a/Engine/Source/Runtime/SignalProcessing/Public/DSP/AudioChannelFormatConverter.h
+++ b/Engine/Source/Runtime/SignalProcessing/Public/DSP/AudioChannelFormatConverter.h
@@ -44,7 +44,7 @@ namespace Audio
 			 * @param InInputBuffers - An array of input audio buffers.
 			 * @param OutOutputBuffers - An array of buffers where output audio is stored. 
 			 */
-			virtual void ProcessAudio(const TArray<AlignedFloatBuffer>& InInputBuffers, TArray<AlignedFloatBuffer>& OutOutputBuffers) = 0;
+			virtual void ProcessAudio(const TArray<FAlignedFloatBuffer>& InInputBuffers, TArray<FAlignedFloatBuffer>& OutOutputBuffers) = 0;
 	};
 
 
@@ -133,7 +133,7 @@ namespace Audio
 			 * @param InInputBuffers - An array of input audio buffers.
 			 * @param OutOutputBuffers - An array of buffers where output audio is stored. 
 			 */
-			void ProcessAudio(const TArray<AlignedFloatBuffer>& InInputBuffers, TArray<AlignedFloatBuffer>& OutOutputBuffers) override;
+			void ProcessAudio(const TArray<FAlignedFloatBuffer>& InInputBuffers, TArray<FAlignedFloatBuffer>& OutOutputBuffers) override;
 
 
 			/** Create a FBaseChannelFormatConverter
diff --git a/Engine/Source/Runtime/SignalProcessing/Public/DSP/AudioFFT.h b/Engine/Source/Runtime/SignalProcessing/Public/DSP/AudioFFT.h
--- a/Engine/Source/Runtime/SignalProcessing/Public/DSP/AudioFFT.h
+++ b/Engine/Source/Runtime/SignalProcessing/Public/DSP/AudioFFT.h
@@ -87,7 +87,7 @@ namespace Audio
 
 	private:
 		EWindowType WindowType;
-		AlignedFloatBuffer WindowBuffer;
+		FAlignedFloatBuffer WindowBuffer;
 		int32 NumSamples;
 
 		// Purposefully hidden constructor.
@@ -174,8 +174,8 @@ namespace Audio
 
 	struct FrequencyBuffer
 	{
-		AlignedFloatBuffer Real;
-		AlignedFloatBuffer Imag;
+		FAlignedFloatBuffer Real;
+		FAlignedFloatBuffer Imag;
 
 		void InitZeroed(int32 Num)
 		{
@@ -204,8 +204,8 @@ namespace Audio
 	// Performs an acyclic FFT correlation on FirstBuffer and Second buffer and stores the output in OutCorrelation.
 	// If bCyclic is false, This function may zero pad FirstBuffer and Second Buffer as needed.
 	// If bCyclic is true, FirstBuffer and SecondBuffer should have the same length, and that length should be a power of two.
-	SIGNALPROCESSING_API void CrossCorrelate(AlignedFloatBuffer& FirstBuffer, AlignedFloatBuffer& SecondBuffer, AlignedFloatBuffer& OutCorrelation, bool bZeroPad = true);
-	SIGNALPROCESSING_API void CrossCorrelate(AlignedFloatBuffer& FirstBuffer, AlignedFloatBuffer& SecondBuffer, FrequencyBuffer& OutCorrelation, bool bZeroPad = true);
+	SIGNALPROCESSING_API void CrossCorrelate(FAlignedFloatBuffer& FirstBuffer, FAlignedFloatBuffer& SecondBuffer, FAlignedFloatBuffer& OutCorrelation, bool bZeroPad = true);
+	SIGNALPROCESSING_API void CrossCorrelate(FAlignedFloatBuffer& FirstBuffer, FAlignedFloatBuffer& SecondBuffer, FrequencyBuffer& OutCorrelation, bool bZeroPad = true);
 	SIGNALPROCESSING_API void CrossCorrelate(const float* FirstBuffer, const float* SecondBuffer, int32 NumSamples, int32 FFTSize, float* OutCorrelation, int32 OutCorrelationSamples);
 	SIGNALPROCESSING_API void CrossCorrelate(const float* FirstBuffer, const float* SecondBuffer, int32 NumSamples, int32 FFTSize, FrequencyBuffer& OutCorrelation);
 
@@ -232,7 +232,7 @@ namespace Audio
 		void SetFilter(const float* InFilterReal, const float* InFilterImag, int32 FilterSize, int32 FFTSize);
 		void SetFilter(const FrequencyBuffer& InFilterFrequencies, int32 FilterSize);
 		void SetFilter(const float* TimeDomainBuffer, int32 FilterSize);
-		void SetFilter(const AlignedFloatBuffer& TimeDomainBuffer);
+		void SetFilter(const FAlignedFloatBuffer& TimeDomainBuffer);
 
 	private:
 		void ConvolveBlock(float* InputAudio, int32 NumSamples);
@@ -244,8 +244,8 @@ namespace Audio
 		int32 BlockSize;
 
 
-		AlignedFloatBuffer TimeDomainInputBuffer;
-		AlignedFloatBuffer COLABuffer;
+		FAlignedFloatBuffer TimeDomainInputBuffer;
+		FAlignedFloatBuffer COLABuffer;
 	};
 
 	// Computes the power spectrum from FFTFreqDomainData. Applies a 1/(FFTSize^2) scaling to the output to 
@@ -255,7 +255,7 @@ namespace Audio
 	// InFrequencyData is the input frequency domain data. Generally this is created by calling PerformFFT(...)
 	// FFTSize is the number of samples used when originally calculating the FFT
 	// OutBuffer is an aligned buffer which will contain spectrum data. It will constain (FFTSize / 2 + 1) elements.
-	SIGNALPROCESSING_API void ComputePowerSpectrum(const FFTFreqDomainData& InFrequencyData, int32 FFTSize, AlignedFloatBuffer& OutBuffer);
+	SIGNALPROCESSING_API void ComputePowerSpectrum(const FFTFreqDomainData& InFrequencyData, int32 FFTSize, FAlignedFloatBuffer& OutBuffer);
 
 	// Computes the magnitude spectrum from FFTFreqDomainData. Applies a 1/FFTSize scaling to the output to 
 	// maintain equal energy between original time domain data and output spectrum.  Only the first 
@@ -264,7 +264,7 @@ namespace Audio
 	// InFrequencyData is the input frequency domain data. Generally this is created by calling PerformFFT(...)
 	// FFTSize is the number of samples used when originally calculating the FFT
 	// OutBuffer is an aligned buffer which will contain spectrum data. It will constain (FFTSize / 2 + 1) elements.
-	SIGNALPROCESSING_API void ComputeMagnitudeSpectrum(const FFTFreqDomainData& InFrequencyData, int32 FFTSize, AlignedFloatBuffer& OutBuffer);
+	SIGNALPROCESSING_API void ComputeMagnitudeSpectrum(const FFTFreqDomainData& InFrequencyData, int32 FFTSize, FAlignedFloatBuffer& OutBuffer);
 	
 	// Computes the spectrum from FFTFreqDomainData. Applies a scaling to the output to maintain equal 
 	// energy between original time domain data and output spectrum.  Only the first (FFTSize / 2 + 1)
@@ -274,5 +274,5 @@ namespace Audio
 	// InFrequencyData is the input frequency domain data. Generally this is created by calling PerformFFT(...)
 	// FFTSize is the number of samples used when originally calculating the FFT
 	// OutBuffer is an aligned buffer which will contain spectrum data. It will constain (FFTSize / 2 + 1) elements.
-	SIGNALPROCESSING_API void ComputeSpectrum(ESpectrumType InSpectrumType, const FFTFreqDomainData& InFrequencyData, int32 FFTSize, AlignedFloatBuffer& OutBuffer);
+	SIGNALPROCESSING_API void ComputeSpectrum(ESpectrumType InSpectrumType, const FFTFreqDomainData& InFrequencyData, int32 FFTSize, FAlignedFloatBuffer& OutBuffer);
 }
diff --git a/Engine/Source/Runtime/SignalProcessing/Public/DSP/BlockCorrelator.h b/Engine/Source/Runtime/SignalProcessing/Public/DSP/BlockCorrelator.h
--- a/Engine/Source/Runtime/SignalProcessing/Public/DSP/BlockCorrelator.h
+++ b/Engine/Source/Runtime/SignalProcessing/Public/DSP/BlockCorrelator.h
@@ -53,7 +53,7 @@ namespace Audio
 			 * @param InputB - Second input block with `GetNumInputValues()` elements. 
 			 * @param Output - Output block with `GetNumOutputValues()` elements.
 			 */
-			void CrossCorrelate(const AlignedFloatBuffer& InputA, const AlignedFloatBuffer& InputB, AlignedFloatBuffer& Output);
+			void CrossCorrelate(const FAlignedFloatBuffer& InputA, const FAlignedFloatBuffer& InputB, FAlignedFloatBuffer& Output);
 
 			/** Autocorrelate an input signal. 
 			 *
@@ -68,15 +68,15 @@ namespace Audio
 			 * @param Input - input block with `GetNumInputValues()` elements. 
 			 * @param Output - Output block with `GetNumOutputValues()` elements.
 			 */
-			void AutoCorrelate(const AlignedFloatBuffer& Input, AlignedFloatBuffer& Output);
+			void AutoCorrelate(const FAlignedFloatBuffer& Input, FAlignedFloatBuffer& Output);
 
 
 		private:
 			// No normalization or windowing.
-			void CyclicCrossCorrelate(const AlignedFloatBuffer& InputA, const AlignedFloatBuffer& InputB, AlignedFloatBuffer& Output);
+			void CyclicCrossCorrelate(const FAlignedFloatBuffer& InputA, const FAlignedFloatBuffer& InputB, FAlignedFloatBuffer& Output);
 
 			// No normalization or windowing.
-			void CyclicAutoCorrelate(const AlignedFloatBuffer& Input, AlignedFloatBuffer& Output);
+			void CyclicAutoCorrelate(const FAlignedFloatBuffer& Input, FAlignedFloatBuffer& Output);
 
 			void InitializeNormalizationBuffer();
 
@@ -89,16 +89,16 @@ namespace Audio
 			TUniquePtr<IFFTAlgorithm> FFTAlgorithm;
 
 			FWindow Window;
-			AlignedFloatBuffer NormalizationBuffer;
+			FAlignedFloatBuffer NormalizationBuffer;
 
-			AlignedFloatBuffer WindowedBufferA;
-			AlignedFloatBuffer WindowedBufferB;
+			FAlignedFloatBuffer WindowedBufferA;
+			FAlignedFloatBuffer WindowedBufferB;
 
-			AlignedFloatBuffer ComplexBufferA;
-			AlignedFloatBuffer ComplexBufferB;
+			FAlignedFloatBuffer ComplexBufferA;
+			FAlignedFloatBuffer ComplexBufferB;
 
-			AlignedFloatBuffer ComplexCorrelationBuffer;
+			FAlignedFloatBuffer ComplexCorrelationBuffer;
 
-			AlignedFloatBuffer FullOutputBuffer;
+			FAlignedFloatBuffer FullOutputBuffer;
 	};
 }
diff --git a/Engine/Source/Runtime/SignalProcessing/Public/DSP/BufferOnePoleLPF.h b/Engine/Source/Runtime/SignalProcessing/Public/DSP/BufferOnePoleLPF.h
--- a/Engine/Source/Runtime/SignalProcessing/Public/DSP/BufferOnePoleLPF.h
+++ b/Engine/Source/Runtime/SignalProcessing/Public/DSP/BufferOnePoleLPF.h
@@ -30,7 +30,7 @@ namespace Audio
 		void SetFrequency(const float InFrequency);
 
 		// Process InSamples and place filtered audio in OutSamples
-		void ProcessAudio(const AlignedFloatBuffer& InSamples, AlignedFloatBuffer& OutSamples);
+		void ProcessAudio(const FAlignedFloatBuffer& InSamples, FAlignedFloatBuffer& OutSamples);
 
 	protected:
 		float CutoffFrequency;
diff --git a/Engine/Source/Runtime/SignalProcessing/Public/DSP/ConstantQ.h b/Engine/Source/Runtime/SignalProcessing/Public/DSP/ConstantQ.h
--- a/Engine/Source/Runtime/SignalProcessing/Public/DSP/ConstantQ.h
+++ b/Engine/Source/Runtime/SignalProcessing/Public/DSP/ConstantQ.h
@@ -80,7 +80,7 @@ namespace Audio
 		static float GetConstantQBandWidth(const float InBandCenter, const float InBandsPerOctave, const float InBandWidthStretch);
 
 		// Fills an array with the spectrum weights to apply to an FFT magnitude or power spectrum output.
-		static void FillArrayWithConstantQBand(const FPseudoConstantQBandSettings& InSettings, AlignedFloatBuffer& OutOffsetArray, int32& OutOffsetIndex);
+		static void FillArrayWithConstantQBand(const FPseudoConstantQBandSettings& InSettings, FAlignedFloatBuffer& OutOffsetArray, int32& OutOffsetIndex);
 
 	};
 }
diff --git a/Engine/Source/Runtime/SignalProcessing/Public/DSP/DynamicDelayAPF.h b/Engine/Source/Runtime/SignalProcessing/Public/DSP/DynamicDelayAPF.h
--- a/Engine/Source/Runtime/SignalProcessing/Public/DSP/DynamicDelayAPF.h
+++ b/Engine/Source/Runtime/SignalProcessing/Public/DSP/DynamicDelayAPF.h
@@ -33,7 +33,7 @@ namespace Audio
 
 		// Processes InSamples through the all pass filter and populates OutSamples with the filter output.
 		// InDelays denotes the per-sample delay of the allpass. It must have an equal number of elements as InSamples
-		void ProcessAudio(const AlignedFloatBuffer& InSamples, const AlignedFloatBuffer& InDelays, AlignedFloatBuffer& OutSamples);
+		void ProcessAudio(const FAlignedFloatBuffer& InSamples, const FAlignedFloatBuffer& InDelays, FAlignedFloatBuffer& OutSamples);
 
 		// Zeros the internal delay line.
 		void Reset();
@@ -43,7 +43,7 @@ namespace Audio
 	protected:
 
 		// Process one block of audio.
-		void ProcessAudioBlock(const float* InSamples, const AlignedFloatBuffer& InFractionalDelays, const int32 InNum, float* OutSamples);
+		void ProcessAudioBlock(const float* InSamples, const FAlignedFloatBuffer& InFractionalDelays, const int32 InNum, float* OutSamples);
 
 	private:
 		// Feedback/Feedforward gain coefficient
@@ -55,10 +55,10 @@ namespace Audio
 		int32 NumInternalBufferSamples;
 
 		// Buffers for block processing
-		AlignedFloatBuffer FractionalDelays;
-		AlignedFloatBuffer DelayLineInput;
-		AlignedFloatBuffer WorkBufferA;
-		AlignedFloatBuffer WorkBufferB;
+		FAlignedFloatBuffer FractionalDelays;
+		FAlignedFloatBuffer DelayLineInput;
+		FAlignedFloatBuffer WorkBufferA;
+		FAlignedFloatBuffer WorkBufferB;
 
 		// Delay line memory.
 		TUniquePtr<FAlignedBlockBuffer> IntegerDelayLine;
diff --git a/Engine/Source/Runtime/SignalProcessing/Public/DSP/EarlyReflectionsFast.h b/Engine/Source/Runtime/SignalProcessing/Public/DSP/EarlyReflectionsFast.h
--- a/Engine/Source/Runtime/SignalProcessing/Public/DSP/EarlyReflectionsFast.h
+++ b/Engine/Source/Runtime/SignalProcessing/Public/DSP/EarlyReflectionsFast.h
@@ -64,7 +64,7 @@ namespace Audio
 		void SetSettings(const FEarlyReflectionsFastSettings& InSettings);
 
 		// Process the single audio frame
-		void ProcessAudio(const AlignedFloatBuffer& InSamples, const int32 InNumChannels, AlignedFloatBuffer& OutLeftSamples, AlignedFloatBuffer& OutRightSamples);
+		void ProcessAudio(const FAlignedFloatBuffer& InSamples, const int32 InNumChannels, FAlignedFloatBuffer& OutLeftSamples, FAlignedFloatBuffer& OutRightSamples);
 
 		// Silence internal audio.
 		void FlushAudio();
@@ -78,12 +78,12 @@ namespace Audio
 		FEarlyReflectionsFastSettings Settings;
 		float SampleRate;
 
-		AlignedFloatBuffer LeftInputBuffer;
-		AlignedFloatBuffer LeftWorkBufferA;
-		AlignedFloatBuffer LeftWorkBufferB;
-		AlignedFloatBuffer RightInputBuffer;
-		AlignedFloatBuffer RightWorkBufferA;
-		AlignedFloatBuffer RightWorkBufferB;
+		FAlignedFloatBuffer LeftInputBuffer;
+		FAlignedFloatBuffer LeftWorkBufferA;
+		FAlignedFloatBuffer LeftWorkBufferB;
+		FAlignedFloatBuffer RightInputBuffer;
+		FAlignedFloatBuffer RightWorkBufferA;
+		FAlignedFloatBuffer RightWorkBufferB;
 
 		FFDNCoefficients LeftCoefficients;
 		FFDNCoefficients RightCoefficients;
diff --git a/Engine/Source/Runtime/SignalProcessing/Public/DSP/FeedbackDelayNetwork.h b/Engine/Source/Runtime/SignalProcessing/Public/DSP/FeedbackDelayNetwork.h
--- a/Engine/Source/Runtime/SignalProcessing/Public/DSP/FeedbackDelayNetwork.h
+++ b/Engine/Source/Runtime/SignalProcessing/Public/DSP/FeedbackDelayNetwork.h
@@ -57,7 +57,7 @@ namespace Audio
 		void SetCoefficients(const FFDNCoefficients& InCoefficients);
 
 		// Generates artificial reverberation for InSamples and places results in OutSamples.
-		void ProcessAudio(const AlignedFloatBuffer& InSamples, AlignedFloatBuffer& OutSamples);
+		void ProcessAudio(const FAlignedFloatBuffer& InSamples, FAlignedFloatBuffer& OutSamples);
 
 		// Sets internal audio samples to silence. 
 		void FlushAudio();
@@ -90,9 +90,9 @@ namespace Audio
 		TUniquePtr<FAlignedBlockBuffer> DelayLine2;
 		TUniquePtr<FAlignedBlockBuffer> DelayLine3;
 
-		AlignedFloatBuffer WorkBuffer0;
-		AlignedFloatBuffer WorkBuffer1;
-		AlignedFloatBuffer WorkBuffer2;
-		AlignedFloatBuffer WorkBuffer3;
+		FAlignedFloatBuffer WorkBuffer0;
+		FAlignedFloatBuffer WorkBuffer1;
+		FAlignedFloatBuffer WorkBuffer2;
+		FAlignedFloatBuffer WorkBuffer3;
 	};
 }
diff --git a/Engine/Source/Runtime/SignalProcessing/Public/DSP/FloatArrayMath.h b/Engine/Source/Runtime/SignalProcessing/Public/DSP/FloatArrayMath.h
--- a/Engine/Source/Runtime/SignalProcessing/Public/DSP/FloatArrayMath.h
+++ b/Engine/Source/Runtime/SignalProcessing/Public/DSP/FloatArrayMath.h
@@ -8,7 +8,7 @@
 namespace Audio
 {
 	SIGNALPROCESSING_API void ArraySum(TArrayView<const float> InValues, float& OutSum);
-	SIGNALPROCESSING_API void ArraySum(const AlignedFloatBuffer& InValues, float& OutSum);
+	SIGNALPROCESSING_API void ArraySum(const FAlignedFloatBuffer& InValues, float& OutSum);
 
 	/** Cumulative sum of array.
 	 *
@@ -87,18 +87,18 @@ namespace Audio
 
 	/** Multiplies two arrays element-wise. Stores result in InValues2*/
 	SIGNALPROCESSING_API void ArrayMultiplyInPlace(TArrayView<const float> InValues1, TArrayView<float> InValues2);
-	SIGNALPROCESSING_API void ArrayMultiplyInPlace(const AlignedFloatBuffer& InValues1, AlignedFloatBuffer& InValues2);
+	SIGNALPROCESSING_API void ArrayMultiplyInPlace(const FAlignedFloatBuffer& InValues1, FAlignedFloatBuffer& InValues2);
 
 	/** Multiplies two complex valued arrays element-wise. 
 	 * This assumes elements are in interleaved format [real_0, imag_0, ..., real_N, imag_N]
 	 * Stores result in InValues2
 	 */
 	SIGNALPROCESSING_API void ArrayComplexMultiplyInPlace(TArrayView<const float> InValues1, TArrayView<float> InValues2);
-	SIGNALPROCESSING_API void ArrayComplexMultiplyInPlace(const AlignedFloatBuffer& InValues1, AlignedFloatBuffer& InValues2);
+	SIGNALPROCESSING_API void ArrayComplexMultiplyInPlace(const FAlignedFloatBuffer& InValues1, FAlignedFloatBuffer& InValues2);
 
 	/** Multiplies each element in InView by InMultiplier */
 	SIGNALPROCESSING_API void ArrayMultiplyByConstantInPlace(TArrayView<float> InValues, float InMultiplier);
-	SIGNALPROCESSING_API void ArrayMultiplyByConstantInPlace(AlignedFloatBuffer& InValues, float InMultiplier);
+	SIGNALPROCESSING_API void ArrayMultiplyByConstantInPlace(FAlignedFloatBuffer& InValues, float InMultiplier);
 
 	/** Add arrays element-wise in place. InAccumulateValues[i] += InValues[i]
 	 *
@@ -106,7 +106,7 @@ namespace Audio
 	 *  InAccumulateValues is the array which holds the sum.
 	 */
 	SIGNALPROCESSING_API void ArrayAddInPlace(TArrayView<const float> InValues, TArrayView<float> InAccumulateValues);
-	SIGNALPROCESSING_API void ArrayAddInPlace(const AlignedFloatBuffer& InValues, AlignedFloatBuffer& InAccumulateValues);
+	SIGNALPROCESSING_API void ArrayAddInPlace(const FAlignedFloatBuffer& InValues, FAlignedFloatBuffer& InAccumulateValues);
 
 	/** Multiply Add arrays element-wise in place. InAccumulateValues[i] += InMultiplier * InValues[i]
 	 *
@@ -115,7 +115,7 @@ namespace Audio
 	 *  @param InAccumulateValues - The array which holds the sum.
 	 */
 	SIGNALPROCESSING_API void ArrayMultiplyAddInPlace(TArrayView<const float> InValues, float InMultiplier, TArrayView<float> InAccumulateValues);
-	SIGNALPROCESSING_API void ArrayMultiplyAddInPlace(const AlignedFloatBuffer& InValues, float InMultiplier, AlignedFloatBuffer& InAccumulateValues);
+	SIGNALPROCESSING_API void ArrayMultiplyAddInPlace(const FAlignedFloatBuffer& InValues, float InMultiplier, FAlignedFloatBuffer& InAccumulateValues);
 
 	/** Linearly Interpolate Add arrays element-wise in place. InAccumulateValues[i] += ((1 - alpha) * InStartMultiplier + alpha * InEndMultipler) * InValues[i]
 	 * Interpolation is performed over the length of the array.
@@ -126,7 +126,7 @@ namespace Audio
 	 *  @param InAccumulateValues - The array which holds the sum.
 	 */
 	SIGNALPROCESSING_API void ArrayLerpAddInPlace(TArrayView<const float> InValues, float InStartMultiplier, float InEndMultiplier, TArrayView<float> InAccumulateValues);
-	SIGNALPROCESSING_API void ArrayLerpAddInPlace(const AlignedFloatBuffer& InValues, float InStartMultiplier, float InEndMultiplier, AlignedFloatBuffer& InAccumulateValues);
+	SIGNALPROCESSING_API void ArrayLerpAddInPlace(const FAlignedFloatBuffer& InValues, float InStartMultiplier, float InEndMultiplier, FAlignedFloatBuffer& InAccumulateValues);
 
 	/** Subract arrays element-wise. OutArray = InMinuend - InSubtrahend
 	 *
@@ -138,7 +138,7 @@ namespace Audio
 
 	/** Subtract value from each element in InValues */
 	SIGNALPROCESSING_API void ArraySubtractByConstantInPlace(TArrayView<float> InValues, float InSubtrahend);
-	SIGNALPROCESSING_API void ArraySubtractByConstantInPlace(AlignedFloatBuffer& InValues, float InSubtrahend);
+	SIGNALPROCESSING_API void ArraySubtractByConstantInPlace(FAlignedFloatBuffer& InValues, float InSubtrahend);
 
 
 	/** Take Square Root of values in place. */
@@ -147,25 +147,25 @@ namespace Audio
 
 	/** Perform complex conjugate of array.  Assumes complex numbers are interlaves [real_0, imag_0, real_1, image_1, ..., real_N, imag_N]. */
 	SIGNALPROCESSING_API void ArrayComplexConjugate(TArrayView<const float> InValues, TArrayView<float> OutValues);
-	SIGNALPROCESSING_API void ArrayComplexConjugate(const AlignedFloatBuffer& InValues, AlignedFloatBuffer& OutValues);
+	SIGNALPROCESSING_API void ArrayComplexConjugate(const FAlignedFloatBuffer& InValues, FAlignedFloatBuffer& OutValues);
 
 	SIGNALPROCESSING_API void ArrayComplexConjugateInPlace(TArrayView<float> InValues);
-	SIGNALPROCESSING_API void ArrayComplexConjugateInPlace(AlignedFloatBuffer& InValues);
+	SIGNALPROCESSING_API void ArrayComplexConjugateInPlace(FAlignedFloatBuffer& InValues);
 
 	/** Convert magnitude values to decibel values in place. db = 20 * log10(val) */
 	SIGNALPROCESSING_API void ArrayMagnitudeToDecibelInPlace(TArrayView<float> InValues, float InMinimumDb);
-	SIGNALPROCESSING_API void ArrayMagnitudeToDecibelInPlace(AlignedFloatBuffer& InValues, float InMinimumDb);
+	SIGNALPROCESSING_API void ArrayMagnitudeToDecibelInPlace(FAlignedFloatBuffer& InValues, float InMinimumDb);
 
 	/** Convert power values to decibel values in place. db = 10 * log10(val) */
 	SIGNALPROCESSING_API void ArrayPowerToDecibelInPlace(TArrayView<float> InValues, float InMinimumDb);
-	SIGNALPROCESSING_API void ArrayPowerToDecibelInPlace(AlignedFloatBuffer& InValues, float InMinimumDb);
+	SIGNALPROCESSING_API void ArrayPowerToDecibelInPlace(FAlignedFloatBuffer& InValues, float InMinimumDb);
 
 
 	/** Compute power of complex data. Out[i] = Complex[2 * i] * Complex[2 * i] + Complex[2 * i + 1] * Complex[2 * i + 1] */
 	SIGNALPROCESSING_API void ArrayComplexToPower(TArrayView<float> InComplexSamples, TArray<float>& OutPowerSamples);
 
 	/** Compute power of complex data. Out[i] = Complex[2 * i] * Complex[2 * i] + Complex[2 * i + 1] * Complex[2 * i + 1] */
-	SIGNALPROCESSING_API void ArrayComplexToPower(const AlignedFloatBuffer& InComplexSamples, AlignedFloatBuffer& OutPowerSamples);
+	SIGNALPROCESSING_API void ArrayComplexToPower(const FAlignedFloatBuffer& InComplexSamples, FAlignedFloatBuffer& OutPowerSamples);
 
 
 
@@ -226,7 +226,7 @@ namespace Audio
 		 *  InView is the array to be transformed. It must have `NumInElements` number of elements.
 		 *  OutArray is the transformed array. It will have `NumOutElements` number of elements.
 		 */
-		void TransformArray(TArrayView<const float> InView, AlignedFloatBuffer& OutArray) const;
+		void TransformArray(TArrayView<const float> InView, FAlignedFloatBuffer& OutArray) const;
 
 		/** Transforms the input array given the kernel.
 		 *
diff --git a/Engine/Source/Runtime/SignalProcessing/Public/DSP/IntegerDelay.h b/Engine/Source/Runtime/SignalProcessing/Public/DSP/IntegerDelay.h
--- a/Engine/Source/Runtime/SignalProcessing/Public/DSP/IntegerDelay.h
+++ b/Engine/Source/Runtime/SignalProcessing/Public/DSP/IntegerDelay.h
@@ -32,11 +32,11 @@ namespace Audio
 		int32 GetNumDelaySamples() const;
 
 		// Process InSamples, placing delayed versions in OutSamples.
-		void ProcessAudio(const Audio::AlignedFloatBuffer& InSamples, Audio::AlignedFloatBuffer& OutSamples);
+		void ProcessAudio(const Audio::FAlignedFloatBuffer& InSamples, Audio::FAlignedFloatBuffer& OutSamples);
 
 		// Retrieve a copy of the internal delay line.
 		// InNum must be less than or equal to InMaxNumDelaySamples.
-		void PeekDelayLine(int32 InNum, Audio::AlignedFloatBuffer& OutSamples);
+		void PeekDelayLine(int32 InNum, Audio::FAlignedFloatBuffer& OutSamples);
 
 	private:
 		// Process a block of audio. InNum is always less than or equal to NumInternalBufferSamples.
diff --git a/Engine/Source/Runtime/SignalProcessing/Public/DSP/LateReflectionsFast.h b/Engine/Source/Runtime/SignalProcessing/Public/DSP/LateReflectionsFast.h
--- a/Engine/Source/Runtime/SignalProcessing/Public/DSP/LateReflectionsFast.h
+++ b/Engine/Source/Runtime/SignalProcessing/Public/DSP/LateReflectionsFast.h
@@ -18,10 +18,10 @@ namespace Audio
 		static const int32 NumTaps = 7;
 
 		// Tap sample buffers.
-		AlignedFloatBuffer Taps[NumTaps];
+		FAlignedFloatBuffer Taps[NumTaps];
 
 		// Output sample buffer.
-		AlignedFloatBuffer Output;
+		FAlignedFloatBuffer Output;
 
 		// Resizes and zero all buffers in structure.
 		void ResizeAndZero(int32 InNumSamples);
@@ -134,9 +134,9 @@ namespace Audio
 		// InDelayModulations contains the fractional delay values per a sample to modulate the first all pass filters delay line.
 		// OutPlateSamples is filled with audio from the various tap points and delay line output.
 		void ProcessAudioFrames(
-			const AlignedFloatBuffer& InSamples,
-			const AlignedFloatBuffer& InFeedbackSamples,
-			const AlignedFloatBuffer& InDelayModulations,
+			const FAlignedFloatBuffer& InSamples,
+			const FAlignedFloatBuffer& InFeedbackSamples,
+			const FAlignedFloatBuffer& InDelayModulations,
 			FLateReflectionsPlateOutputs& OutPlateSamples);
 
 		// Flush internal audio to silence.
@@ -153,7 +153,7 @@ namespace Audio
 		// Retrieve the internal delay line data. The number of samples retrieved is limited to the
 		// number of samples in the internal buffer. That amount can be determined by calling
 		// "GetNumInternalBufferSamples()".
-		void PeekDelayLine(int32 InNum, AlignedFloatBuffer& OutSamples);
+		void PeekDelayLine(int32 InNum, FAlignedFloatBuffer& OutSamples);
 
 	protected:
 		// Current parameter settings of reverb
@@ -182,9 +182,9 @@ namespace Audio
 		TUniquePtr<FIntegerDelay> DelayH;
 		TUniquePtr<FIntegerDelay> DelayI;
 
-		Audio::AlignedFloatBuffer WorkBufferA;
-		Audio::AlignedFloatBuffer WorkBufferB;
-		Audio::AlignedFloatBuffer WorkBufferC;
+		Audio::FAlignedFloatBuffer WorkBufferA;
+		Audio::FAlignedFloatBuffer WorkBufferB;
+		Audio::FAlignedFloatBuffer WorkBufferC;
 	};
 
 	// Settings for controlling the FLateReflections
@@ -258,7 +258,7 @@ namespace Audio
 		// Create reverberation in OutSamples based upon InSamples. 
 		// OutSamples is a 2 channel audio buffer no matter the value of InNumChannels.
 		// OutSamples is mixed with InSamples based upon gain.
-		void ProcessAudio(const AlignedFloatBuffer& InSamples, const int32 InNumChannels, AlignedFloatBuffer& OutLeftSamples, AlignedFloatBuffer& OutRightSamples);
+		void ProcessAudio(const FAlignedFloatBuffer& InSamples, const int32 InNumChannels, FAlignedFloatBuffer& OutLeftSamples, FAlignedFloatBuffer& OutRightSamples);
 
 		// Flush internal audio to silence.
 		void FlushAudio();
@@ -268,7 +268,7 @@ namespace Audio
 
 		void ApplySettings();
 
-		void GeneraterPlateModulations(const int32 InNum, AlignedFloatBuffer& OutLeftDelays, AlignedFloatBuffer& OutRightDelays);
+		void GeneraterPlateModulations(const int32 InNum, FAlignedFloatBuffer& OutLeftDelays, FAlignedFloatBuffer& OutRightDelays);
 
 		float SampleRate;
 		float Gain;
@@ -306,12 +306,12 @@ namespace Audio
 		FLateReflectionsPlateOutputs RightPlateOutputs;
 
 		// Buffer for delay modulation of left and right plates.
-		AlignedFloatBuffer LeftDelayModSamples;
-		AlignedFloatBuffer RightDelayModSamples;
+		FAlignedFloatBuffer LeftDelayModSamples;
+		FAlignedFloatBuffer RightDelayModSamples;
 
-		AlignedFloatBuffer WorkBufferA;
-		AlignedFloatBuffer WorkBufferB;
-		AlignedFloatBuffer WorkBufferC;
+		FAlignedFloatBuffer WorkBufferA;
+		FAlignedFloatBuffer WorkBufferB;
+		FAlignedFloatBuffer WorkBufferC;
 
 	};
 
diff --git a/Engine/Source/Runtime/SignalProcessing/Public/DSP/LinearInterpFractionalDelay.h b/Engine/Source/Runtime/SignalProcessing/Public/DSP/LinearInterpFractionalDelay.h
--- a/Engine/Source/Runtime/SignalProcessing/Public/DSP/LinearInterpFractionalDelay.h
+++ b/Engine/Source/Runtime/SignalProcessing/Public/DSP/LinearInterpFractionalDelay.h
@@ -19,7 +19,7 @@ namespace Audio
 
 		// Apply delay to InSamples. Fill OutSamples with data from the delay line at a delay of InDelay.
 		// InDelay must be equal length to InSamples. 
-		void ProcessAudio(const AlignedFloatBuffer& InSamples, const AlignedFloatBuffer& InDelays, AlignedFloatBuffer& OutSamples);
+		void ProcessAudio(const FAlignedFloatBuffer& InSamples, const FAlignedFloatBuffer& InDelays, FAlignedFloatBuffer& OutSamples);
 
 		// Set all values in internal delay line to zero. 
 		void Reset();
diff --git a/Engine/Source/Runtime/SignalProcessing/Public/DSP/LinkwitzRileyBandSplitter.h b/Engine/Source/Runtime/SignalProcessing/Public/DSP/LinkwitzRileyBandSplitter.h
--- a/Engine/Source/Runtime/SignalProcessing/Public/DSP/LinkwitzRileyBandSplitter.h
+++ b/Engine/Source/Runtime/SignalProcessing/Public/DSP/LinkwitzRileyBandSplitter.h
@@ -22,7 +22,7 @@ namespace Audio
 
 	struct FMultibandBuffer
 	{
-		Audio::AlignedFloatBuffer Buffer;
+		Audio::FAlignedFloatBuffer Buffer;
 
 		int32 NumBands = 0;
 		int32 NumSamples = 0;
diff --git a/Engine/Source/Runtime/SignalProcessing/Public/DSP/LongDelayAPF.h b/Engine/Source/Runtime/SignalProcessing/Public/DSP/LongDelayAPF.h
--- a/Engine/Source/Runtime/SignalProcessing/Public/DSP/LongDelayAPF.h
+++ b/Engine/Source/Runtime/SignalProcessing/Public/DSP/LongDelayAPF.h
@@ -25,11 +25,11 @@ namespace Audio
 		void SetG(float InG) { G = InG; }
 
 		// Process InSamples and place filtered data in OutSamples
-		void ProcessAudio(const AlignedFloatBuffer& InSamples, AlignedFloatBuffer& OutSamples);
+		void ProcessAudio(const FAlignedFloatBuffer& InSamples, FAlignedFloatBuffer& OutSamples);
 
 		// Process InSamples and place filtered data in OutSamples and delay line samples in OutDelaySamples.
 		// The OutDelaySamples correspond to "w[n] = InSamples[n] + InG * w[n - InDelay]"
-		void ProcessAudio(const AlignedFloatBuffer& InSamples, AlignedFloatBuffer& OutSamples, AlignedFloatBuffer& OutDelaySamples);
+		void ProcessAudio(const FAlignedFloatBuffer& InSamples, FAlignedFloatBuffer& OutSamples, FAlignedFloatBuffer& OutDelaySamples);
 
 		// Sets delay line values to zero.
 		void Reset();
@@ -51,7 +51,7 @@ namespace Audio
 		int32 NumInternalBufferSamples;
 
 		// Buffer for block operations.
-		AlignedFloatBuffer WorkBuffer;
+		FAlignedFloatBuffer WorkBuffer;
 
 		// Filter delay line memory.
 		TUniquePtr<FAlignedBlockBuffer> DelayLine;
diff --git a/Engine/Source/Runtime/SignalProcessing/Public/DSP/MultithreadedPatching.h b/Engine/Source/Runtime/SignalProcessing/Public/DSP/MultithreadedPatching.h
--- a/Engine/Source/Runtime/SignalProcessing/Public/DSP/MultithreadedPatching.h
+++ b/Engine/Source/Runtime/SignalProcessing/Public/DSP/MultithreadedPatching.h
@@ -54,7 +54,7 @@ namespace Audio
 		TCircularAudioBuffer<float> InternalBuffer;
 
 		// For MixInAudio, audio is popped off of InternalBuffer onto here, and then mixed into OutBuffer in MixInAudio.
-		AlignedFloatBuffer MixingBuffer;
+		FAlignedFloatBuffer MixingBuffer;
 
 		// This is applied in PopAudio or MixInAudio.
 		TAtomic<float> TargetGain;
@@ -252,6 +252,6 @@ namespace Audio
 		FPatchSplitter Splitter;
 
 		/** This buffer is used to pop audio from our Mixer and push to our splitter. */
-		AlignedFloatBuffer IntermediateBuffer;
+		FAlignedFloatBuffer IntermediateBuffer;
 	};
 }
diff --git a/Engine/Source/Runtime/SignalProcessing/Public/DSP/PassiveFilter.h b/Engine/Source/Runtime/SignalProcessing/Public/DSP/PassiveFilter.h
--- a/Engine/Source/Runtime/SignalProcessing/Public/DSP/PassiveFilter.h
+++ b/Engine/Source/Runtime/SignalProcessing/Public/DSP/PassiveFilter.h
@@ -127,8 +127,8 @@ namespace Audio
 			return;
 		}
 
-		AlignedFloatBuffer TempReal;
-		AlignedFloatBuffer TempImag;
+		FAlignedFloatBuffer TempReal;
+		FAlignedFloatBuffer TempImag;
 		TempReal.AddUninitialized(Signal.Num());
 		TempImag.AddUninitialized(Signal.Num());
 
diff --git a/Engine/Source/Runtime/SignalProcessing/Public/DSP/ReverbFast.h b/Engine/Source/Runtime/SignalProcessing/Public/DSP/ReverbFast.h
--- a/Engine/Source/Runtime/SignalProcessing/Public/DSP/ReverbFast.h
+++ b/Engine/Source/Runtime/SignalProcessing/Public/DSP/ReverbFast.h
@@ -71,7 +71,7 @@ namespace Audio
 			// Creates reverberated audio in OutSamples based upon InSamples
 			// InNumChannels can be 1 or 2 channels.
 			// OutSamples must be greater or equal to 2.
-			void ProcessAudio(const AlignedFloatBuffer& InSamples, const int32 InNumChannels, AlignedFloatBuffer& OutSamples, const int32 OutNumChannels);
+			void ProcessAudio(const FAlignedFloatBuffer& InSamples, const int32 InNumChannels, FAlignedFloatBuffer& OutSamples, const int32 OutNumChannels);
 
 			void FlushAudio();
 
@@ -80,7 +80,7 @@ namespace Audio
 
 		private:
 			// Copy reverberated samples to interleaved output samples. Map channels according to internal settings.
-			void InterleaveAndMixOutput(const AlignedFloatBuffer& InFrontLeftSamples, const AlignedFloatBuffer& InFrontRightSamples, AlignedFloatBuffer& OutSamples, const int32 OutNumChannels);
+			void InterleaveAndMixOutput(const FAlignedFloatBuffer& InFrontLeftSamples, const FAlignedFloatBuffer& InFrontRightSamples, FAlignedFloatBuffer& OutSamples, const int32 OutNumChannels);
 
 			void ApplySettings();
 
@@ -91,15 +91,15 @@ namespace Audio
 			FEarlyReflectionsFast EarlyReflections;
 			FLateReflectionsFast LateReflections;
 
-			AlignedFloatBuffer FrontLeftLateReflectionsSamples;
-			AlignedFloatBuffer FrontRightLateReflectionsSamples;
-			AlignedFloatBuffer FrontLeftEarlyReflectionsSamples;
-			AlignedFloatBuffer FrontRightEarlyReflectionsSamples;
-			AlignedFloatBuffer FrontLeftReverbSamples;
-			AlignedFloatBuffer FrontRightReverbSamples;
-			AlignedFloatBuffer LeftAttenuatedSamples;
-			AlignedFloatBuffer RightAttenuatedSamples;
-			AlignedFloatBuffer ScaledInputBuffer;
+			FAlignedFloatBuffer FrontLeftLateReflectionsSamples;
+			FAlignedFloatBuffer FrontRightLateReflectionsSamples;
+			FAlignedFloatBuffer FrontLeftEarlyReflectionsSamples;
+			FAlignedFloatBuffer FrontRightEarlyReflectionsSamples;
+			FAlignedFloatBuffer FrontLeftReverbSamples;
+			FAlignedFloatBuffer FrontRightReverbSamples;
+			FAlignedFloatBuffer LeftAttenuatedSamples;
+			FAlignedFloatBuffer RightAttenuatedSamples;
+			FAlignedFloatBuffer ScaledInputBuffer;
 
 	};
 }
diff --git a/Engine/Source/Runtime/SignalProcessing/Public/DSP/SpectrumAnalyzer.h b/Engine/Source/Runtime/SignalProcessing/Public/DSP/SpectrumAnalyzer.h
--- a/Engine/Source/Runtime/SignalProcessing/Public/DSP/SpectrumAnalyzer.h
+++ b/Engine/Source/Runtime/SignalProcessing/Public/DSP/SpectrumAnalyzer.h
@@ -197,7 +197,7 @@ namespace Audio
 			 * @param InTimestamp - A timestamp associated with the input complex buffer.
 			 * @param OutValues - Array to store output bands.
 			 */
-			virtual void ExtractBands(const AlignedFloatBuffer& InComplexBuffer, double InTimestamp, TArray<float>& OutValues) = 0;
+			virtual void ExtractBands(const FAlignedFloatBuffer& InComplexBuffer, double InTimestamp, TArray<float>& OutValues) = 0;
 
 			/** Creates a ISpectrumBandExtractor. */
 			static TUniquePtr<ISpectrumBandExtractor> CreateSpectrumBandExtractor(const FSpectrumBandExtractorSettings& InSettings);
@@ -216,21 +216,21 @@ namespace Audio
 		void Reset(int32 InNum);
 
 		// Input. Used on analysis thread to lock a buffer to write to.
-		AlignedFloatBuffer& StartWorkOnBuffer();
+		FAlignedFloatBuffer& StartWorkOnBuffer();
 
 		// When calling stop work on buffer, also set timestmap associated with buffer.
 		void StopWorkOnBuffer(double InTimestamp);
 		
 		// Output. Used to lock the most recent buffer we analyzed.
-		const AlignedFloatBuffer& LockMostRecentBuffer() const;
+		const FAlignedFloatBuffer& LockMostRecentBuffer() const;
 
 		// Output. Used to lock the most recent buffer we analyzed.
 		// OutTimestamp is populated with the timestamp associated with the buffer wehn StopWorkOnBuffer is called.
-		const AlignedFloatBuffer& LockMostRecentBuffer(double& OutTimestamp) const;
+		const FAlignedFloatBuffer& LockMostRecentBuffer(double& OutTimestamp) const;
 		void UnlockBuffer();
 
 	private:
-		TArray<AlignedFloatBuffer> ComplexBuffers;
+		TArray<FAlignedFloatBuffer> ComplexBuffers;
 		TArray<double> Timestamps;
 
 		// Private functions. Either increments or decrements the respective counter,
@@ -356,7 +356,7 @@ namespace Audio
 		void ResetSettings();
 
 		// Called in GetMagnitudeForFrequency and GetPhaseForFrequency.
-		void PerformInterpolation(const AlignedFloatBuffer& InComplexBuffer, EPeakInterpolationMethod InMethod, const float InFreq, float& OutReal, float& OutImag);
+		void PerformInterpolation(const FAlignedFloatBuffer& InComplexBuffer, EPeakInterpolationMethod InMethod, const float InFreq, float& OutReal, float& OutImag);
 
 		// Cached current settings. Only actually used in ResetSettings().
 		FSpectrumAnalyzerSettings CurrentSettings;
@@ -372,14 +372,14 @@ namespace Audio
 		int32 HopInSamples;
 		EFFTScaling FFTScaling;
 
-		AlignedFloatBuffer AnalysisTimeDomainBuffer;
+		FAlignedFloatBuffer AnalysisTimeDomainBuffer;
 		FThreadSafeCounter SampleCounter; 
 		TCircularAudioBuffer<float> InputQueue;
 		FSpectrumAnalyzerBuffer FrequencyBuffer;
 
 		// if non-null, owns pointer to locked frequency vector we're using.
 		double LockedBufferTimestamp;
-		const AlignedFloatBuffer* LockedFrequencyVector;
+		const FAlignedFloatBuffer* LockedFrequencyVector;
 
 
 		TUniquePtr<IFFTAlgorithm> FFT;
diff --git a/Engine/Source/Runtime/SignalProcessing/Public/DSP/VoiceProcessing.h b/Engine/Source/Runtime/SignalProcessing/Public/DSP/VoiceProcessing.h
--- a/Engine/Source/Runtime/SignalProcessing/Public/DSP/VoiceProcessing.h
+++ b/Engine/Source/Runtime/SignalProcessing/Public/DSP/VoiceProcessing.h
@@ -185,8 +185,8 @@ namespace Audio
 		FrequencyBuffer IncomingFrequencies;
 		FrequencyBuffer OutgoingFrequencies;
 
-		AlignedFloatBuffer ZeroPaddedIncomingBuffer;
-		AlignedFloatBuffer ZeroPaddedOutgoingBuffer;
+		FAlignedFloatBuffer ZeroPaddedIncomingBuffer;
+		FAlignedFloatBuffer ZeroPaddedOutgoingBuffer;
 	};
 
 	/*
@@ -214,7 +214,7 @@ namespace Audio
 		FFDAPFilterComputer FilterComputer;
 		FAdaptiveFilter AdaptiveFilter;
 		
-		AlignedFloatBuffer FilterComputerInput;
+		FAlignedFloatBuffer FilterComputerInput;
 		FrequencyBuffer FilterComputerOutput;
 		float ConvergenceRate;
 		int32 FilterLength;
diff --git a/Engine/Source/Runtime/SignalProcessing/Public/SampleBuffer.h b/Engine/Source/Runtime/SignalProcessing/Public/SampleBuffer.h
--- a/Engine/Source/Runtime/SignalProcessing/Public/SampleBuffer.h
+++ b/Engine/Source/Runtime/SignalProcessing/Public/SampleBuffer.h
@@ -60,7 +60,7 @@ namespace Audio
 			FMemory::Memcpy(RawPCMData.GetData(), Other.RawPCMData.GetData(), NumSamples * sizeof(SampleType));
 		}
 
-		FORCEINLINE TSampleBuffer(const AlignedFloatBuffer& InData, int32 InNumChannels, int32 InSampleRate)
+		FORCEINLINE TSampleBuffer(const FAlignedFloatBuffer& InData, int32 InNumChannels, int32 InSampleRate)
 		{
 			*this =  TSampleBuffer(InData.GetData(), InData.Num(), InNumChannels, InSampleRate);
 		}
diff --git a/Engine/Source/Runtime/SoundFieldRendering/Private/SoundFieldRendering.cpp b/Engine/Source/Runtime/SoundFieldRendering/Private/SoundFieldRendering.cpp
--- a/Engine/Source/Runtime/SoundFieldRendering/Private/SoundFieldRendering.cpp
+++ b/Engine/Source/Runtime/SoundFieldRendering/Private/SoundFieldRendering.cpp
@@ -81,7 +81,7 @@ FSoundFieldDecoder::FSoundFieldDecoder()
 	TargetSpeakerGains = FoaVirtualSpeakerWordLockedGains;
 }
 
-void FSoundFieldDecoder::DecodeAudioDirectlyToDeviceOutputPositions(const FAmbisonicsSoundfieldBuffer& InputData, const FSoundfieldSpeakerPositionalData& OutputPositions, Audio::AlignedFloatBuffer& OutputData)
+void FSoundFieldDecoder::DecodeAudioDirectlyToDeviceOutputPositions(const FAmbisonicsSoundfieldBuffer& InputData, const FSoundfieldSpeakerPositionalData& OutputPositions, Audio::FAlignedFloatBuffer& OutputData)
 {
 	if (InputData.NumChannels == 0 || InputData.AudioBuffer.Num() == 0)
 	{
@@ -279,7 +279,7 @@ void FSoundFieldDecoder::EvenOrderDecodeLoop(const int32 NumFrames, const int32
 
 
 // this will probably be the use case for 5.1, prioritize optimizations here
-void FSoundFieldDecoder::DecodeAudioToSevenOneAndDownmixToDevice(const FAmbisonicsSoundfieldBuffer& InputData, const FSoundfieldSpeakerPositionalData& OutputPositions, Audio::AlignedFloatBuffer& OutputData)
+void FSoundFieldDecoder::DecodeAudioToSevenOneAndDownmixToDevice(const FAmbisonicsSoundfieldBuffer& InputData, const FSoundfieldSpeakerPositionalData& OutputPositions, Audio::FAlignedFloatBuffer& OutputData)
 {
 	if (InputData.NumChannels == 0 || InputData.AudioBuffer.Num() == 0)
 	{
@@ -384,7 +384,7 @@ void FSoundFieldDecoder::DecodeAudioToSevenOneAndDownmixToDevice(const FAmbisoni
 	else // Interpolated decode
 	{
 		SpeakerGainLerper.Init(CurrentSpeakerGains, TargetSpeakerGains, NumFrames);
-		const Audio::AlignedFloatBuffer& SpeakerGainLerpDeltas = SpeakerGainLerper.GetDeltaBuffer();
+		const Audio::FAlignedFloatBuffer& SpeakerGainLerpDeltas = SpeakerGainLerper.GetDeltaBuffer();
 		const float* RESTRICT SpeakerGainLerpDeltasPtr = SpeakerGainLerpDeltas.GetData();
 		VectorRegister CurrSpeakerGainDelta = VectorRegister();
 
@@ -654,7 +654,7 @@ TArray<Audio::FChannelPositionInfo>* FSoundFieldDecoder::GetDefaultChannelPositi
 	return nullptr;
 }
 
-void FSoundFieldDecoder::FoaRotationInPlace(Audio::AlignedFloatBuffer& InOutBuffer, const float XRotDegrees, const float YRotDegrees, const float ZRotDegrees)
+void FSoundFieldDecoder::FoaRotationInPlace(Audio::FAlignedFloatBuffer& InOutBuffer, const float XRotDegrees, const float YRotDegrees, const float ZRotDegrees)
 {
 	const int32 NumSamples = InOutBuffer.Num();
 	const int32 NumFrames = NumSamples / 4; // FOA
@@ -702,7 +702,7 @@ FSoundFieldEncoder::FSoundFieldEncoder()
 {
 }
 
-void FSoundFieldEncoder::EncodeAudioDirectlyFromOutputPositions(const Audio::AlignedFloatBuffer& InputData, const FSoundfieldSpeakerPositionalData& InputPositions, const FAmbisonicsSoundfieldSettings& Settings, FAmbisonicsSoundfieldBuffer& OutputData)
+void FSoundFieldEncoder::EncodeAudioDirectlyFromOutputPositions(const Audio::FAlignedFloatBuffer& InputData, const FSoundfieldSpeakerPositionalData& InputPositions, const FAmbisonicsSoundfieldSettings& Settings, FAmbisonicsSoundfieldBuffer& OutputData)
 {
 	constexpr const float DEG_2_RAD = PI / 180.0f;
 
diff --git a/Engine/Source/Runtime/SoundFieldRendering/Public/SoundFieldRendering.h b/Engine/Source/Runtime/SoundFieldRendering/Public/SoundFieldRendering.h
--- a/Engine/Source/Runtime/SoundFieldRendering/Public/SoundFieldRendering.h
+++ b/Engine/Source/Runtime/SoundFieldRendering/Public/SoundFieldRendering.h
@@ -16,7 +16,7 @@ class SOUNDFIELDRENDERING_API FAmbisonicsSoundfieldBuffer : public ISoundfieldAu
 {
 public:
 	// Interleaved audio buffer for all vector parts of the ambisonics stream.
-	Audio::AlignedFloatBuffer AudioBuffer;
+	Audio::FAlignedFloatBuffer AudioBuffer;
 
 	// number of channels in the Ambisonics stream.
 	// Currently we don't explicitly support mixed order ambisonics,
@@ -57,11 +57,11 @@ class SOUNDFIELDRENDERING_API FSoundFieldDecoder
 public:
 	FSoundFieldDecoder(); // initializes VirtualSpeakerWorldLockedSpeakerGains
 
-	void DecodeAudioDirectlyToDeviceOutputPositions(const FAmbisonicsSoundfieldBuffer& InputData, const FSoundfieldSpeakerPositionalData& OutputPositions, Audio::AlignedFloatBuffer& OutputData);
-	void DecodeAudioToSevenOneAndDownmixToDevice(const FAmbisonicsSoundfieldBuffer& InputData, const FSoundfieldSpeakerPositionalData& OutputPositions, Audio::AlignedFloatBuffer& OutputData);
+	void DecodeAudioDirectlyToDeviceOutputPositions(const FAmbisonicsSoundfieldBuffer& InputData, const FSoundfieldSpeakerPositionalData& OutputPositions, Audio::FAlignedFloatBuffer& OutputData);
+	void DecodeAudioToSevenOneAndDownmixToDevice(const FAmbisonicsSoundfieldBuffer& InputData, const FSoundfieldSpeakerPositionalData& OutputPositions, Audio::FAlignedFloatBuffer& OutputData);
 	static void RotateFirstOrderAmbisonicsBed(const FAmbisonicsSoundfieldBuffer& InputData, FAmbisonicsSoundfieldBuffer& OutputData, const FQuat& DestinationRotation, const FQuat& PreviousRotation);
 
-	static void FoaRotationInPlace(Audio::AlignedFloatBuffer& InOutFrames, const float XRotDegrees, const float YRotDegrees, const float ZRotDegrees);
+	static void FoaRotationInPlace(Audio::FAlignedFloatBuffer& InOutFrames, const float XRotDegrees, const float YRotDegrees, const float ZRotDegrees);
 
 private:
 	// Special cased, vectorized versions of first order operations.
@@ -79,11 +79,11 @@ private:
 	// Helper function to get the default channel positions for a channel array.
 	static TArray<Audio::FChannelPositionInfo>* GetDefaultChannelPositions(int32 InNumChannels);
 
-	Audio::AlignedFloatBuffer VirtualSpeakerScratchBuffers;
-	Audio::AlignedFloatBuffer FoaVirtualSpeakerWordLockedGains;
-	Audio::AlignedFloatBuffer TargetSpeakerGains;
-	Audio::AlignedFloatBuffer CurrentSpeakerGains;
-	Audio::AlignedFloatBuffer MixdownGainsMap;
+	Audio::FAlignedFloatBuffer VirtualSpeakerScratchBuffers;
+	Audio::FAlignedFloatBuffer FoaVirtualSpeakerWordLockedGains;
+	Audio::FAlignedFloatBuffer TargetSpeakerGains;
+	Audio::FAlignedFloatBuffer CurrentSpeakerGains;
+	Audio::FAlignedFloatBuffer MixdownGainsMap;
 	Audio::FBufferLinearEase SpeakerGainLerper;
 	FVector2D LastListenerRotationSphericalCoord{ -999.f, -999.f }; // force interpolation on first callback
 
@@ -100,12 +100,12 @@ class SOUNDFIELDRENDERING_API FSoundFieldEncoder
 {
 public:
 	FSoundFieldEncoder();
-	void EncodeAudioDirectlyFromOutputPositions(const Audio::AlignedFloatBuffer& InputData, const FSoundfieldSpeakerPositionalData& InputPositions, const FAmbisonicsSoundfieldSettings& Settings, FAmbisonicsSoundfieldBuffer& OutputData);
+	void EncodeAudioDirectlyFromOutputPositions(const Audio::FAlignedFloatBuffer& InputData, const FSoundfieldSpeakerPositionalData& InputPositions, const FAmbisonicsSoundfieldSettings& Settings, FAmbisonicsSoundfieldBuffer& OutputData);
 
 private:
 	static void EncodeLoop(const int32 NumFrames, const int32 NumInputChannels, const float* RESTRICT InputAudioPtr, const int32 NumAmbiChannels, float* RESTRICT SpeakerGainsPtr, float* RESTRICT OutputAmbiBuffer);
 
-	Audio::AlignedFloatBuffer VirtualSpeakerScratchBuffers;
-	Audio::AlignedFloatBuffer SpeakerGains;
-	Audio::AlignedFloatBuffer MixdownGainsMap;
+	Audio::FAlignedFloatBuffer VirtualSpeakerScratchBuffers;
+	Audio::FAlignedFloatBuffer SpeakerGains;
+	Audio::FAlignedFloatBuffer MixdownGainsMap;
 };
