Fix for inconsisnent use of defines for vector register alignment (elements vs bytes)

diff --git a/Engine/Plugins/Runtime/Metasound/Source/MetasoundGraphCore/Private/MetasoundOperatorSettings.cpp b/Engine/Plugins/Runtime/Metasound/Source/MetasoundGraphCore/Private/MetasoundOperatorSettings.cpp
--- a/Engine/Plugins/Runtime/Metasound/Source/MetasoundGraphCore/Private/MetasoundOperatorSettings.cpp
+++ b/Engine/Plugins/Runtime/Metasound/Source/MetasoundGraphCore/Private/MetasoundOperatorSettings.cpp
@@ -14,7 +14,7 @@ namespace Metasound
 		static constexpr float MinimumTargetBlockRate = 1.0f;
 
 		// Using same alignment as in BufferVectorOperations.h
-		static constexpr int32 FloatAlignment = AUDIO_SIMD_FLOAT_ALIGNMENT;
+		static constexpr int32 FloatAlignment = AUDIO_NUM_FLOATS_PER_VECTOR_REGISTER;
 
 		static constexpr int32 MinimumNumFramesPerBlock = FloatAlignment;
 		static constexpr int32 MaximumNumFramesPerBlock = 16384;
diff --git a/Engine/Plugins/Runtime/Metasound/Source/MetasoundStandardNodes/Private/MetasoundMathNodes.cpp b/Engine/Plugins/Runtime/Metasound/Source/MetasoundStandardNodes/Private/MetasoundMathNodes.cpp
--- a/Engine/Plugins/Runtime/Metasound/Source/MetasoundStandardNodes/Private/MetasoundMathNodes.cpp
+++ b/Engine/Plugins/Runtime/Metasound/Source/MetasoundStandardNodes/Private/MetasoundMathNodes.cpp
@@ -708,7 +708,7 @@ namespace Metasound
 				const float* OperandData = InAdditionalOperands[i]->GetData();
 				float* OutData = OutResult->GetData();
 
-				if (NumSamples % AUDIO_SIMD_FLOAT_ALIGNMENT)
+				if (NumSamples % AUDIO_NUM_FLOATS_PER_VECTOR_REGISTER)
 				{
 					for (int32 SampleIndex = 0; SampleIndex < NumSamples; ++SampleIndex)
 					{
@@ -802,7 +802,7 @@ namespace Metasound
 		{
 			FMemory::Memcpy(OutResult->GetData(), InPrimaryOperand->GetData(), sizeof(float) * OutResult->Num());
 
-			const int32 SIMDRemainder = OutResult->Num() % AUDIO_SIMD_FLOAT_ALIGNMENT;
+			const int32 SIMDRemainder = OutResult->Num() % AUDIO_NUM_FLOATS_PER_VECTOR_REGISTER;
 			const int32 SIMDCount = OutResult->Num() - SIMDRemainder;
 
 			for (int32 i = 0; i < InAdditionalOperands.Num(); ++i)
@@ -926,7 +926,7 @@ namespace Metasound
 				const float* OperandData = InAdditionalOperands[i]->GetData();
 				float* OutData = OutResult->GetData();
 
-				if (NumSamples % AUDIO_SIMD_FLOAT_ALIGNMENT)
+				if (NumSamples % AUDIO_NUM_FLOATS_PER_VECTOR_REGISTER)
 				{
 					for (int32 SampleIndex = 0; SampleIndex < NumSamples; ++SampleIndex)
 					{
@@ -1094,7 +1094,7 @@ namespace Metasound
 				const float* OperandData = InAdditionalOperands[i]->GetData();
 				float* OutData = OutResult->GetData();
 
-				if (NumSamples % AUDIO_SIMD_FLOAT_ALIGNMENT)
+				if (NumSamples % AUDIO_NUM_FLOATS_PER_VECTOR_REGISTER)
 				{
 					for (int32 SampleIndex = 0; SampleIndex < NumSamples; ++SampleIndex)
 					{
@@ -1159,7 +1159,7 @@ namespace Metasound
 			float NewGain = 1.0f;
 			for (int32 i = 0; i < InAdditionalOperands.Num(); ++i)
 			{
-				const int32 SIMDRemainder = OutResult->Num() % AUDIO_SIMD_FLOAT_ALIGNMENT;
+				const int32 SIMDRemainder = OutResult->Num() % AUDIO_NUM_FLOATS_PER_VECTOR_REGISTER;
 				const int32 SIMDCount = OutResult->Num() - SIMDRemainder;
 
 				Audio::FadeBufferFast(OutResult->GetData(), SIMDCount, InInstanceData.LastGain, *InAdditionalOperands[i]);
diff --git a/Engine/Plugins/Runtime/Metasound/Source/MetasoundStandardNodes/Private/MetasoundSampleAndHoldNode.cpp b/Engine/Plugins/Runtime/Metasound/Source/MetasoundStandardNodes/Private/MetasoundSampleAndHoldNode.cpp
--- a/Engine/Plugins/Runtime/Metasound/Source/MetasoundStandardNodes/Private/MetasoundSampleAndHoldNode.cpp
+++ b/Engine/Plugins/Runtime/Metasound/Source/MetasoundStandardNodes/Private/MetasoundSampleAndHoldNode.cpp
@@ -128,37 +128,10 @@ namespace Metasound
 
 				const int32 NumSamples = EndFrame - StartFrame;
 
-				// if CurrOutputPtr is on the last float of a (4)aligned block, StartAlignment will be 3,
-				// and we want LeadingNonSIMDCount below to be 1 (4-1)
-				const int32 StartAlignment = (uintptr_t)((void*)CurrOutputPtr) % AUDIO_SIMD_FLOAT_ALIGNMENT;
-
-				// if start alignment is 0, we are already aligned (hence the !!)
-				int32 SamplesRemaining = NumSamples;
-				const int32 LeadingNonSIMDCount = (StartAlignment == 0)? 0 : (AUDIO_SIMD_FLOAT_ALIGNMENT - StartAlignment);
-
-				SamplesRemaining -= LeadingNonSIMDCount;
-				const int32 SIMDCount = SamplesRemaining - (SamplesRemaining % AUDIO_SIMD_FLOAT_ALIGNMENT);
-
-				SamplesRemaining -= SIMDCount;
-				const int32 TrailingNonSIMDCount = SamplesRemaining;
-
-
-				// leading non-simd
-				for (int32 i = 0; i < LeadingNonSIMDCount; ++i)
-				{
-					*(CurrOutputPtr++) = HoldValue;
-				}
-
-				// simd
-				if (SIMDCount)
-				{
-					Audio::BufferSetToConstantInplace(CurrOutputPtr, SIMDCount, HoldValue);
-				}
-
-				// trailing non-simd
-				for (int32 i = 0; i < TrailingNonSIMDCount; ++i)
+				// non-SIMD version
+				for (int32 i = 0; i < NumSamples; ++i)
 				{
-					*(CurrOutputPtr++) = HoldValue;
+					CurrOutputPtr[i] = HoldValue;
 				}
 			};
 
diff --git a/Engine/Source/Runtime/SignalProcessing/Private/BufferVectorOperations.cpp b/Engine/Source/Runtime/SignalProcessing/Private/BufferVectorOperations.cpp
--- a/Engine/Source/Runtime/SignalProcessing/Private/BufferVectorOperations.cpp
+++ b/Engine/Source/Runtime/SignalProcessing/Private/BufferVectorOperations.cpp
@@ -21,14 +21,14 @@ namespace Audio
 	/* Sets a values to zero if value is denormal. Denormal numbers significantly slow down floating point operations. */
 	void BufferUnderflowClampFast(float* RESTRICT InOutBuffer, const int32 InNum)
 	{
-		checkf(InNum >= 4, TEXT("Buffer must have atleast 4 elements."));
-		checkf(0 == (InNum % 4), TEXT("Buffer length be a multiple of 4."));
-		checkf(IsAligned<float*>(InOutBuffer, AUDIO_SIMD_FLOAT_ALIGNMENT), TEXT("Memory must be aligned to use vector operations."));
+		checkf(InNum >= AUDIO_NUM_FLOATS_PER_VECTOR_REGISTER, TEXT("Buffer must have at least %i elements."), AUDIO_NUM_FLOATS_PER_VECTOR_REGISTER);
+		checkf(0 == (InNum % AUDIO_NUM_FLOATS_PER_VECTOR_REGISTER), TEXT("Buffer length be a multiple of %i."), AUDIO_NUM_FLOATS_PER_VECTOR_REGISTER);
+		checkf(IsAligned<float*>(InOutBuffer, AUDIO_SIMD_BYTE_ALIGNMENT), TEXT("Memory must be aligned to use vector operations."));
 
 		const VectorRegister4Float VFMIN = MakeVectorRegister(FLT_MIN, FLT_MIN, FLT_MIN, FLT_MIN);
 		const VectorRegister4Float VNFMIN = MakeVectorRegister(-FLT_MIN, -FLT_MIN, -FLT_MIN, -FLT_MIN);
 
-		for (int32 i = 0; i < InNum; i += 4)
+		for (int32 i = 0; i < InNum; i += AUDIO_NUM_FLOATS_PER_VECTOR_REGISTER)
 		{
 			VectorRegister4Float VInOut = VectorLoadAligned(&InOutBuffer[i]);
 
@@ -50,14 +50,14 @@ namespace Audio
 	/* Clamps values in the buffer to be between InMinValue and InMaxValue */
 	void BufferRangeClampFast(float* RESTRICT InOutBuffer, const int32 InNum, float InMinValue, float InMaxValue)
 	{
-		checkf(InNum >= 4, TEXT("Buffer must have atleast 4 elements."));
-		checkf(0 == (InNum % 4), TEXT("Buffer length be a multiple of 4."));
-		checkf(IsAligned<float*>(InOutBuffer, AUDIO_SIMD_FLOAT_ALIGNMENT), TEXT("Memory must be aligned to use vector operations."));
+		checkf(InNum >= AUDIO_NUM_FLOATS_PER_VECTOR_REGISTER, TEXT("Buffer must have at least %i elements."), AUDIO_NUM_FLOATS_PER_VECTOR_REGISTER);
+		checkf(0 == (InNum % AUDIO_NUM_FLOATS_PER_VECTOR_REGISTER), TEXT("Buffer length be a multiple of %i."), AUDIO_NUM_FLOATS_PER_VECTOR_REGISTER);
+		checkf(IsAligned<float*>(InOutBuffer, AUDIO_SIMD_BYTE_ALIGNMENT), TEXT("Memory must be aligned to use vector operations."));
 
 		const VectorRegister4Float VMinVal = MakeVectorRegister(InMinValue, InMinValue, InMinValue, InMinValue);
 		const VectorRegister4Float VMaxVal = MakeVectorRegister(InMaxValue, InMaxValue, InMaxValue, InMaxValue);
 
-		for (int32 i = 0; i < InNum; i += 4)
+		for (int32 i = 0; i < InNum; i += AUDIO_NUM_FLOATS_PER_VECTOR_REGISTER)
 		{
 			VectorRegister4Float VInOut = VectorLoadAligned(&InOutBuffer[i]);
 
@@ -101,12 +101,12 @@ namespace Audio
 		RestrictedPtrAliasCheck(InFloatBuffer, OutFloatBuffer, InNumSamples);
 
 		// Can only SIMD on multiple of 4 buffers, we'll do normal multiples on last bit
-		const int32 NumSamplesRemaining = InNumSamples % 4;
+		const int32 NumSamplesRemaining = InNumSamples % AUDIO_NUM_FLOATS_PER_VECTOR_REGISTER;
 		const int32 NumSamplesToSimd = InNumSamples - NumSamplesRemaining;
 
 		// Load the single value we want to multiply all values by into a vector register
 		const VectorRegister4Float MultiplyValue = VectorLoadFloat1(&InValue);
-		for (int32 i = 0; i < NumSamplesToSimd; i += 4)
+		for (int32 i = 0; i < NumSamplesToSimd; i += AUDIO_NUM_FLOATS_PER_VECTOR_REGISTER)
 		{
 			// Load the next 4 samples of the input buffer into a register
 			VectorRegister4Float InputBufferRegister = VectorLoadAligned(&InFloatBuffer[i]);
@@ -135,7 +135,7 @@ namespace Audio
 	{
 		const VectorRegister4Float Gain = VectorLoadFloat1(&InGain);
 
-		for (int32 i = 0; i < NumSamples; i += 4)
+		for (int32 i = 0; i < NumSamples; i += AUDIO_NUM_FLOATS_PER_VECTOR_REGISTER)
 		{
 			VectorRegister4Float Output = VectorLoadAligned(&InBuffer[i]);
 			Output = VectorMultiply(Output, Gain);
@@ -153,7 +153,7 @@ namespace Audio
 	{
 		const VectorRegister4Float Constant = VectorLoadFloat1(&InConstant);
 
-		for (int32 i = 0; i < NumSamples; i += 4)
+		for (int32 i = 0; i < NumSamples; i += AUDIO_NUM_FLOATS_PER_VECTOR_REGISTER)
 		{
 			VectorRegister4Float Output = VectorLoadAligned(&InBuffer[i]);
 			Output = VectorAdd(Output, Constant);
@@ -168,13 +168,13 @@ namespace Audio
 
 	void BufferSetToConstantInplace(float* RESTRICT InBuffer, int32 NumSamples, float InConstant)
 	{
-		checkf(NumSamples >= 4, TEXT("Buffer must have atleast 4 elements."));
-		checkf(0 == (NumSamples % 4), TEXT("Buffer length be a multiple of 4."));
-		checkf(IsAligned<const float*>(InBuffer, AUDIO_SIMD_FLOAT_ALIGNMENT), TEXT("Memory must be aligned to use vector operations."));
+		checkf(NumSamples >= AUDIO_NUM_FLOATS_PER_VECTOR_REGISTER, TEXT("Buffer must have at least %i elements."), AUDIO_NUM_FLOATS_PER_VECTOR_REGISTER);
+		checkf(0 == (NumSamples % AUDIO_NUM_FLOATS_PER_VECTOR_REGISTER), TEXT("Buffer length be a multiple of %i."), AUDIO_NUM_FLOATS_PER_VECTOR_REGISTER);
+		checkf(IsAligned<const float*>(InBuffer, AUDIO_SIMD_BYTE_ALIGNMENT), TEXT("Memory must be aligned to use vector operations."));
 
 		const VectorRegister4Float Constant = VectorLoadFloat1(&InConstant);
 
-		for (int32 i = 0; i < NumSamples; i += 4)
+		for (int32 i = 0; i < NumSamples; i += AUDIO_NUM_FLOATS_PER_VECTOR_REGISTER)
 		{
 			VectorStoreAligned(Constant, &InBuffer[i]);
 		}
@@ -203,11 +203,11 @@ namespace Audio
 	/* Performs an element-wise weighted sum OutputBuffer = (InBuffer1 x InGain1) + (InBuffer2 x InGain2) */
 	void BufferWeightedSumFast(const float* RESTRICT InBuffer1, float InGain1, const float* RESTRICT InBuffer2, float InGain2, float* RESTRICT OutBuffer, int32 InNum)
 	{
-		checkf(InNum >= 4, TEXT("Buffer must have atleast 4 elements."));
-		checkf(0 == (InNum % 4), TEXT("Buffer length be a multiple of 4."));
-		checkf(IsAligned<const float*>(InBuffer1, AUDIO_SIMD_FLOAT_ALIGNMENT), TEXT("Memory must be aligned to use vector operations."));
-		checkf(IsAligned<const float*>(InBuffer2, AUDIO_SIMD_FLOAT_ALIGNMENT), TEXT("Memory must be aligned to use vector operations."));
-		checkf(IsAligned<float*>(OutBuffer, AUDIO_SIMD_FLOAT_ALIGNMENT), TEXT("Memory must be aligned to use vector operations."));
+		checkf(InNum >= AUDIO_NUM_FLOATS_PER_VECTOR_REGISTER, TEXT("Buffer must have at least %i elements."), AUDIO_NUM_FLOATS_PER_VECTOR_REGISTER);
+		checkf(0 == (InNum % AUDIO_NUM_FLOATS_PER_VECTOR_REGISTER), TEXT("Buffer length be a multiple of %i."), AUDIO_NUM_FLOATS_PER_VECTOR_REGISTER);
+		checkf(IsAligned<const float*>(InBuffer1, AUDIO_SIMD_BYTE_ALIGNMENT), TEXT("Memory must be aligned to use vector operations."));
+		checkf(IsAligned<const float*>(InBuffer2, AUDIO_SIMD_BYTE_ALIGNMENT), TEXT("Memory must be aligned to use vector operations."));
+		checkf(IsAligned<float*>(OutBuffer, AUDIO_SIMD_BYTE_ALIGNMENT), TEXT("Memory must be aligned to use vector operations."));
 
 		RestrictedPtrAliasCheck(InBuffer1, OutBuffer, InNum);
 		RestrictedPtrAliasCheck(InBuffer2, OutBuffer, InNum);
@@ -215,7 +215,7 @@ namespace Audio
 		VectorRegister4Float Gain1Vector = VectorLoadFloat1(&InGain1);
 		VectorRegister4Float Gain2Vector = VectorLoadFloat1(&InGain2);
 
-		for (int32 i = 0; i < InNum; i += 4)
+		for (int32 i = 0; i < InNum; i += AUDIO_NUM_FLOATS_PER_VECTOR_REGISTER)
 		{
 			// InBuffer1 x InGain1
 			VectorRegister4Float Input1 = VectorLoadAligned(&InBuffer1[i]);
@@ -233,18 +233,18 @@ namespace Audio
 	/* Performs an element-wise weighted sum OutputBuffer = (InBuffer1 x InGain1) + InBuffer2 */
 	void BufferWeightedSumFast(const float* RESTRICT InBuffer1, float InGain1, const float* RESTRICT InBuffer2, float* RESTRICT OutBuffer, int32 InNum)
 	{
-		checkf(InNum >= 4, TEXT("Buffer must have atleast 4 elements."));
-		checkf(0 == (InNum % 4), TEXT("Buffer length be a multiple of 4."));
-		checkf(IsAligned<const float*>(InBuffer1, AUDIO_SIMD_FLOAT_ALIGNMENT), TEXT("Memory must be aligned to use vector operations."));
-		checkf(IsAligned<const float*>(InBuffer2, AUDIO_SIMD_FLOAT_ALIGNMENT), TEXT("Memory must be aligned to use vector operations."));
-		checkf(IsAligned<float*>(OutBuffer, AUDIO_SIMD_FLOAT_ALIGNMENT), TEXT("Memory must be aligned to use vector operations."));
+		checkf(InNum >= AUDIO_NUM_FLOATS_PER_VECTOR_REGISTER, TEXT("Buffer must have at least %i elements."), AUDIO_NUM_FLOATS_PER_VECTOR_REGISTER);
+		checkf(0 == (InNum % AUDIO_NUM_FLOATS_PER_VECTOR_REGISTER), TEXT("Buffer length be a multiple of %i."), AUDIO_NUM_FLOATS_PER_VECTOR_REGISTER);
+		checkf(IsAligned<const float*>(InBuffer1, AUDIO_SIMD_BYTE_ALIGNMENT), TEXT("Memory must be aligned to use vector operations."));
+		checkf(IsAligned<const float*>(InBuffer2, AUDIO_SIMD_BYTE_ALIGNMENT), TEXT("Memory must be aligned to use vector operations."));
+		checkf(IsAligned<float*>(OutBuffer, AUDIO_SIMD_BYTE_ALIGNMENT), TEXT("Memory must be aligned to use vector operations."));
 
 		RestrictedPtrAliasCheck(InBuffer1, OutBuffer, InNum);
 		RestrictedPtrAliasCheck(InBuffer2, OutBuffer, InNum);
 
 		VectorRegister4Float Gain1Vector = VectorLoadFloat1(&InGain1);
 
-		for (int32 i = 0; i < InNum; i += 4)
+		for (int32 i = 0; i < InNum; i += AUDIO_NUM_FLOATS_PER_VECTOR_REGISTER)
 		{
 			// InBuffer1 x InGain1
 			VectorRegister4Float Input1 = VectorLoadAligned(&InBuffer1[i]);
@@ -264,8 +264,8 @@ namespace Audio
 
 	void FadeBufferFast(float* RESTRICT OutFloatBuffer, int32 NumSamples, const float StartValue, const float EndValue)
 	{
-		checkf(IsAligned<float*>(OutFloatBuffer, 4), TEXT("Memory must be aligned to use vector operations."));
-		checkf(NumSamples % 4 == 0, TEXT("Please use a buffer size that is a multiple of 4."));
+		checkf(IsAligned<float*>(OutFloatBuffer, AUDIO_SIMD_BYTE_ALIGNMENT), TEXT("Memory must be aligned to use vector operations."));
+		checkf(NumSamples % AUDIO_NUM_FLOATS_PER_VECTOR_REGISTER == 0, TEXT("Please use a buffer size that is a multiple of %i."), AUDIO_NUM_FLOATS_PER_VECTOR_REGISTER);
 
 #if !AUDIO_USE_SIMD
 		float Gain = StartValue;
@@ -309,7 +309,7 @@ namespace Audio
 			{
 				VectorRegister4Float Gain = VectorLoadFloat1(&StartValue);
 
-				for (int32 i = 0; i < NumSamples; i += 4)
+				for (int32 i = 0; i < NumSamples; i += AUDIO_NUM_FLOATS_PER_VECTOR_REGISTER)
 				{
 					VectorRegister4Float Output = VectorLoadAligned(&OutFloatBuffer[i]);
 					Output = VectorMultiply(Output, Gain);
@@ -324,7 +324,7 @@ namespace Audio
 			VectorRegister4Float Gain = VectorLoadFloat1(&StartValue);
 			VectorRegister4Float Delta = VectorLoadFloat1(&DeltaValue);
 
-			for (int32 i = 0; i < NumSamples; i += 4)
+			for (int32 i = 0; i < NumSamples; i += AUDIO_NUM_FLOATS_PER_VECTOR_REGISTER)
 			{
 				VectorRegister4Float Output = VectorLoadAligned(&OutFloatBuffer[i]);
 				Output = VectorMultiply(Output, Gain);
@@ -342,9 +342,9 @@ namespace Audio
 
 	void MixInBufferFast(const float* RESTRICT InFloatBuffer, float* RESTRICT BufferToSumTo, int32 NumSamples, const float Gain)
 	{
-		checkf(IsAligned<const float*>(InFloatBuffer, 4), TEXT("Memory must be aligned to use vector operations."));
-		checkf(IsAligned<float*>(BufferToSumTo, 4), TEXT("Memory must be aligned to use vector operations."));
-		checkf(NumSamples % 4 == 0, TEXT("Please use a buffer size that is a multiple of 4."));
+		checkf(IsAligned<const float*>(InFloatBuffer, AUDIO_SIMD_BYTE_ALIGNMENT), TEXT("Memory must be aligned to use vector operations."));
+		checkf(IsAligned<float*>(BufferToSumTo, AUDIO_SIMD_BYTE_ALIGNMENT), TEXT("Memory must be aligned to use vector operations."));
+		checkf(NumSamples % AUDIO_NUM_FLOATS_PER_VECTOR_REGISTER == 0, TEXT("Please use a buffer size that is a multiple of %i."), AUDIO_NUM_FLOATS_PER_VECTOR_REGISTER);
 
 		RestrictedPtrAliasCheck(InFloatBuffer, BufferToSumTo, NumSamples);
 
@@ -356,7 +356,7 @@ namespace Audio
 #else
 		VectorRegister4Float GainVector = VectorLoadFloat1(&Gain);
 
-		for (int32 i = 0; i < NumSamples; i += 4)
+		for (int32 i = 0; i < NumSamples; i += AUDIO_NUM_FLOATS_PER_VECTOR_REGISTER)
 		{
 			VectorRegister4Float Output = VectorLoadAligned(&BufferToSumTo[i]);
 			VectorRegister4Float Input  = VectorLoadAligned(&InFloatBuffer[i]);
@@ -374,9 +374,9 @@ namespace Audio
 
 	void MixInBufferFast(const float* RESTRICT InFloatBuffer, float* RESTRICT BufferToSumTo, int32 NumSamples)
 	{
-		checkf(IsAligned<const float*>(InFloatBuffer, 4), TEXT("Memory must be aligned to use vector operations."));
-		checkf(IsAligned<float*>(BufferToSumTo, 4), TEXT("Memory must be aligned to use vector operations."));
-		checkf(NumSamples % 4 == 0, TEXT("Please use a buffer size that is a multiple of 4."));
+		checkf(IsAligned<const float*>(InFloatBuffer, AUDIO_SIMD_BYTE_ALIGNMENT), TEXT("Memory must be aligned to use vector operations."));
+		checkf(IsAligned<float*>(BufferToSumTo, AUDIO_SIMD_BYTE_ALIGNMENT), TEXT("Memory must be aligned to use vector operations."));
+		checkf(NumSamples % AUDIO_NUM_FLOATS_PER_VECTOR_REGISTER == 0, TEXT("Please use a buffer size that is a multiple of %i."), AUDIO_NUM_FLOATS_PER_VECTOR_REGISTER);
 
 		RestrictedPtrAliasCheck(InFloatBuffer, BufferToSumTo, NumSamples);
 
@@ -386,7 +386,7 @@ namespace Audio
 			BufferToSumTo[i] += InFloatBuffer[i];
 		}
 #else
-		for (int32 i = 0; i < NumSamples; i += 4)
+		for (int32 i = 0; i < NumSamples; i += AUDIO_NUM_FLOATS_PER_VECTOR_REGISTER)
 		{
 			VectorRegister4Float Output = VectorLoadAligned(&BufferToSumTo[i]);
 			VectorRegister4Float Input = VectorLoadAligned(&InFloatBuffer[i]);
@@ -403,9 +403,9 @@ namespace Audio
 
 	void MixInBufferFast(const float* RESTRICT InFloatBuffer, float* RESTRICT BufferToSumTo, int32 NumSamples, const float StartGain, const float EndGain)
 	{
-		checkf(IsAligned<const float*>(InFloatBuffer, 4), TEXT("Memory must be aligned to use vector operations."));
-		checkf(IsAligned<float*>(BufferToSumTo, 4), TEXT("Memory must be aligned to use vector operations."));
-		checkf(NumSamples % 4 == 0, TEXT("Please use a buffer size that is a multiple of 4."));
+		checkf(IsAligned<const float*>(InFloatBuffer, AUDIO_SIMD_BYTE_ALIGNMENT), TEXT("Memory must be aligned to use vector operations."));
+		checkf(IsAligned<float*>(BufferToSumTo, AUDIO_SIMD_BYTE_ALIGNMENT), TEXT("Memory must be aligned to use vector operations."));
+		checkf(NumSamples % AUDIO_NUM_FLOATS_PER_VECTOR_REGISTER == 0, TEXT("Please use a buffer size that is a multiple of %i."), AUDIO_NUM_FLOATS_PER_VECTOR_REGISTER);
 
 		RestrictedPtrAliasCheck(InFloatBuffer, BufferToSumTo, NumSamples);
 
@@ -422,7 +422,7 @@ namespace Audio
 			{
 				VectorRegister4Float Gain = VectorLoadFloat1(&StartGain);
 
-				for (int32 i = 0; i < NumSamples; i += 4)
+				for (int32 i = 0; i < NumSamples; i += AUDIO_NUM_FLOATS_PER_VECTOR_REGISTER)
 				{
 					VectorRegister4Float Input = VectorLoadAligned(&InFloatBuffer[i]);
 					VectorRegister4Float Output = VectorLoadAligned(&BufferToSumTo[i]);
@@ -441,7 +441,7 @@ namespace Audio
 			VectorRegister4Float Gain = VectorLoadFloat1(&StartGain);
 			VectorRegister4Float Delta = VectorLoadFloat1(&DeltaValue);
 
-			for (int32 i = 0; i < NumSamples; i += 4)
+			for (int32 i = 0; i < NumSamples; i += AUDIO_NUM_FLOATS_PER_VECTOR_REGISTER)
 			{
 				VectorRegister4Float Input = VectorLoadAligned(&InFloatBuffer[i]);
 				VectorRegister4Float Output = VectorLoadAligned(&BufferToSumTo[i]);
@@ -470,16 +470,16 @@ namespace Audio
 	/* Subtracts two buffers together element-wise. */
 	void BufferSubtractFast(const float* RESTRICT InMinuend, const float* RESTRICT InSubtrahend, float* RESTRICT OutBuffer, int32 InNum)
 	{
-		checkf(InNum >= 4, TEXT("Buffer must have atleast 4 elements."));
-		checkf(0 == (InNum % 4), TEXT("Buffer length be a multiple of 4."));
-		checkf(IsAligned<const float*>(InMinuend, AUDIO_SIMD_FLOAT_ALIGNMENT), TEXT("Memory must be aligned to use vector operations."));
-		checkf(IsAligned<const float*>(InSubtrahend, AUDIO_SIMD_FLOAT_ALIGNMENT), TEXT("Memory must be aligned to use vector operations."));
-		checkf(IsAligned<float*>(OutBuffer, AUDIO_SIMD_FLOAT_ALIGNMENT), TEXT("Memory must be aligned to use vector operations."));
+		checkf(InNum >= 4, TEXT("Buffer must have at least %i elements."), AUDIO_NUM_FLOATS_PER_VECTOR_REGISTER);
+		checkf(0 == (InNum % AUDIO_NUM_FLOATS_PER_VECTOR_REGISTER), TEXT("Buffer length be a multiple of %i."), AUDIO_NUM_FLOATS_PER_VECTOR_REGISTER);
+		checkf(IsAligned<const float*>(InMinuend, AUDIO_SIMD_BYTE_ALIGNMENT), TEXT("Memory must be aligned to use vector operations."));
+		checkf(IsAligned<const float*>(InSubtrahend, AUDIO_SIMD_BYTE_ALIGNMENT), TEXT("Memory must be aligned to use vector operations."));
+		checkf(IsAligned<float*>(OutBuffer, AUDIO_SIMD_BYTE_ALIGNMENT), TEXT("Memory must be aligned to use vector operations."));
 
 		RestrictedPtrAliasCheck(InMinuend, OutBuffer, InNum);
 		RestrictedPtrAliasCheck(InSubtrahend, OutBuffer, InNum);
 
-		for (int32 i = 0; i < InNum; i += 4)
+		for (int32 i = 0; i < InNum; i += AUDIO_NUM_FLOATS_PER_VECTOR_REGISTER)
 		{
 			VectorRegister4Float Input1 = VectorLoadAligned(&InMinuend[i]);
 			VectorRegister4Float Input2 = VectorLoadAligned(&InSubtrahend[i]);
@@ -498,14 +498,14 @@ namespace Audio
 	/* Performs element-wise in-place subtraction placing the result in the subtrahend. InOutSubtrahend = InMinuend - InOutSubtrahend */
 	void BufferSubtractInPlace1Fast(const float* RESTRICT InMinuend, float* RESTRICT InOutSubtrahend, int32 InNum)
 	{
-		checkf(InNum >= 4, TEXT("Buffer must have atleast 4 elements."));
-		checkf(0 == (InNum % 4), TEXT("Buffer length be a multiple of 4."));
-		checkf(IsAligned<const float*>(InMinuend, AUDIO_SIMD_FLOAT_ALIGNMENT), TEXT("Memory must be aligned to use vector operations."));
-		checkf(IsAligned<const float*>(InOutSubtrahend, AUDIO_SIMD_FLOAT_ALIGNMENT), TEXT("Memory must be aligned to use vector operations."));
+		checkf(InNum >= 4, TEXT("Buffer must have at least %i elements."), AUDIO_NUM_FLOATS_PER_VECTOR_REGISTER);
+		checkf(0 == (InNum % AUDIO_NUM_FLOATS_PER_VECTOR_REGISTER), TEXT("Buffer length be a multiple of %i."), AUDIO_NUM_FLOATS_PER_VECTOR_REGISTER);
+		checkf(IsAligned<const float*>(InMinuend, AUDIO_SIMD_BYTE_ALIGNMENT), TEXT("Memory must be aligned to use vector operations."));
+		checkf(IsAligned<const float*>(InOutSubtrahend, AUDIO_SIMD_BYTE_ALIGNMENT), TEXT("Memory must be aligned to use vector operations."));
 		
 		RestrictedPtrAliasCheck(InMinuend, InOutSubtrahend, InNum);
 
-		for (int32 i = 0; i < InNum; i += 4)
+		for (int32 i = 0; i < InNum; i += AUDIO_NUM_FLOATS_PER_VECTOR_REGISTER)
 		{
 			VectorRegister4Float Input1 = VectorLoadAligned(&InMinuend[i]);
 			VectorRegister4Float Input2 = VectorLoadAligned(&InOutSubtrahend[i]);
@@ -525,14 +525,14 @@ namespace Audio
 	/* Performs element-wise in-place subtraction placing the result in the minuend. InOutMinuend = InOutMinuend - InSubtrahend */
 	void BufferSubtractInPlace2Fast(float* RESTRICT InOutMinuend, const float* RESTRICT InSubtrahend, int32 InNum)
 	{
-		checkf(InNum >= 4, TEXT("Buffer must have atleast 4 elements."));
-		checkf(0 == (InNum % 4), TEXT("Buffer length be a multiple of 4."));
-		checkf(IsAligned<const float*>(InOutMinuend, AUDIO_SIMD_FLOAT_ALIGNMENT), TEXT("Memory must be aligned to use vector operations."));
-		checkf(IsAligned<const float*>(InSubtrahend, AUDIO_SIMD_FLOAT_ALIGNMENT), TEXT("Memory must be aligned to use vector operations."));
+		checkf(InNum >= 4, TEXT("Buffer must have at least %i elements."), AUDIO_NUM_FLOATS_PER_VECTOR_REGISTER);
+		checkf(0 == (InNum % AUDIO_NUM_FLOATS_PER_VECTOR_REGISTER), TEXT("Buffer length be a multiple of %i."), AUDIO_NUM_FLOATS_PER_VECTOR_REGISTER);
+		checkf(IsAligned<const float*>(InOutMinuend, AUDIO_SIMD_BYTE_ALIGNMENT), TEXT("Memory must be aligned to use vector operations."));
+		checkf(IsAligned<const float*>(InSubtrahend, AUDIO_SIMD_BYTE_ALIGNMENT), TEXT("Memory must be aligned to use vector operations."));
 
 		RestrictedPtrAliasCheck(InOutMinuend, InSubtrahend, InNum);
 		
-		for (int32 i = 0; i < InNum; i += 4)
+		for (int32 i = 0; i < InNum; i += AUDIO_NUM_FLOATS_PER_VECTOR_REGISTER)
 		{
 			VectorRegister4Float Input1 = VectorLoadAligned(&InOutMinuend[i]);
 			VectorRegister4Float Input2 = VectorLoadAligned(&InSubtrahend[i]);
@@ -554,10 +554,10 @@ namespace Audio
 
 	void SumBuffers(const float* RESTRICT InFloatBuffer1, const float* RESTRICT InFloatBuffer2, float* RESTRICT OutputBuffer, int32 NumSamples)
 	{
-		checkf(IsAligned<const float*>(InFloatBuffer1, 4), TEXT("Memory must be aligned to use vector operations."));
-		checkf(IsAligned<const float*>(InFloatBuffer2, 4), TEXT("Memory must be aligned to use vector operations."));
-		checkf(IsAligned<float*>(OutputBuffer, 4), TEXT("Memory must be aligned to use vector operations."));
-		checkf(NumSamples % 4 == 0, TEXT("Please use a buffer size that is a multiple of 4."));
+		checkf(IsAligned<const float*>(InFloatBuffer1, AUDIO_SIMD_BYTE_ALIGNMENT), TEXT("Memory must be aligned to use vector operations."));
+		checkf(IsAligned<const float*>(InFloatBuffer2, AUDIO_SIMD_BYTE_ALIGNMENT), TEXT("Memory must be aligned to use vector operations."));
+		checkf(IsAligned<float*>(OutputBuffer, AUDIO_SIMD_BYTE_ALIGNMENT), TEXT("Memory must be aligned to use vector operations."));
+		checkf(NumSamples % AUDIO_NUM_FLOATS_PER_VECTOR_REGISTER == 0, TEXT("Please use a buffer size that is a multiple of %i."), AUDIO_NUM_FLOATS_PER_VECTOR_REGISTER);
 
 		RestrictedPtrAliasCheck(InFloatBuffer1, OutputBuffer, NumSamples);
 		RestrictedPtrAliasCheck(InFloatBuffer2, OutputBuffer, NumSamples);
@@ -568,7 +568,7 @@ namespace Audio
 			OutputBuffer[i] = InFloatBuffer1[i] + InFloatBuffer2[i];
 		}
 #else
-		for (int32 i = 0; i < NumSamples; i += 4)
+		for (int32 i = 0; i < NumSamples; i += AUDIO_NUM_FLOATS_PER_VECTOR_REGISTER)
 		{
 			VectorRegister4Float Input1 = VectorLoadAligned(&InFloatBuffer1[i]);
 			VectorRegister4Float Input2 = VectorLoadAligned(&InFloatBuffer2[i]);
@@ -586,13 +586,13 @@ namespace Audio
 
 	void MultiplyBuffersInPlace(const float* RESTRICT InFloatBuffer, float* RESTRICT BufferToMultiply, int32 NumSamples)
 	{
-		checkf(IsAligned<const float*>(InFloatBuffer, 4), TEXT("Memory must be aligned to use vector operations."));
-		checkf(IsAligned<const float*>(BufferToMultiply, 4), TEXT("Memory must be aligned to use vector operations."));
-		checkf(NumSamples % 4 == 0, TEXT("Please use a buffer size that is a multiple of 4."));
+		checkf(IsAligned<const float*>(InFloatBuffer, AUDIO_SIMD_BYTE_ALIGNMENT), TEXT("Memory must be aligned to use vector operations."));
+		checkf(IsAligned<const float*>(BufferToMultiply, AUDIO_SIMD_BYTE_ALIGNMENT), TEXT("Memory must be aligned to use vector operations."));
+		checkf(NumSamples % AUDIO_NUM_FLOATS_PER_VECTOR_REGISTER == 0, TEXT("Please use a buffer size that is a multiple of %i."), AUDIO_NUM_FLOATS_PER_VECTOR_REGISTER);
 
 		RestrictedPtrAliasCheck(InFloatBuffer, BufferToMultiply, NumSamples);
 
-		for (int32 i = 0; i < NumSamples; i += 4)
+		for (int32 i = 0; i < NumSamples; i += AUDIO_NUM_FLOATS_PER_VECTOR_REGISTER)
 		{
 			VectorRegister4Float Input1 = VectorLoadAligned(&InFloatBuffer[i]);
 			VectorRegister4Float Output = VectorLoadAligned(&BufferToMultiply[i]);
@@ -609,7 +609,7 @@ namespace Audio
 
 	float GetMagnitude(const float* RESTRICT Buffer, int32 NumSamples)
 	{
-		checkf(NumSamples % 4 == 0, TEXT("Please use a buffer size that is a multiple of 4."));
+		checkf(NumSamples % AUDIO_NUM_FLOATS_PER_VECTOR_REGISTER == 0, TEXT("Please use a buffer size that is a multiple of %i."), AUDIO_NUM_FLOATS_PER_VECTOR_REGISTER);
 
 #if !AUDIO_USE_SIMD
 		float Sum = 0.0f;
@@ -624,7 +624,7 @@ namespace Audio
 		const float Exponent = 2.0f;
 		VectorRegister4Float ExponentVector = VectorLoadFloat1(&Exponent);
 
-		for (int32 i = 0; i < NumSamples; i += 4)
+		for (int32 i = 0; i < NumSamples; i += AUDIO_NUM_FLOATS_PER_VECTOR_REGISTER)
 		{
 			VectorRegister4Float Input = VectorPow(VectorLoadAligned(&Buffer[i]), ExponentVector);
 			Sum = VectorAdd(Sum, Input);
@@ -639,14 +639,14 @@ namespace Audio
 
 	float BufferGetAverageValue(const FAlignedFloatBuffer& Buffer)
 	{
-		checkf(Buffer.Num() % 4 == 0, TEXT("Please use a buffer size that is a multiple of 4."));
+		checkf(Buffer.Num() % AUDIO_NUM_FLOATS_PER_VECTOR_REGISTER == 0, TEXT("Please use a buffer size that is a multiple of %i."), AUDIO_NUM_FLOATS_PER_VECTOR_REGISTER);
 
 		return BufferGetAverageValue(Buffer.GetData(), Buffer.Num());
 	}
 
 	float BufferGetAverageValue(const float* RESTRICT Buffer, int32 NumSamples)
 	{
-		checkf(NumSamples % 4 == 0, TEXT("Please use a buffer size that is a multiple of 4."));
+		checkf(NumSamples % AUDIO_NUM_FLOATS_PER_VECTOR_REGISTER == 0, TEXT("Please use a buffer size that is a multiple of %i."), AUDIO_NUM_FLOATS_PER_VECTOR_REGISTER);
 
 #if !AUDIO_USE_SIMD
 		float Sum = 0.0f;
@@ -658,7 +658,7 @@ namespace Audio
 #else
 		VectorRegister4Float Sum = VectorZero();
 
-		for (int32 i = 0; i < NumSamples; i += 4)
+		for (int32 i = 0; i < NumSamples; i += AUDIO_NUM_FLOATS_PER_VECTOR_REGISTER)
 		{
 			VectorRegister4Float Input = VectorLoadAligned(&Buffer[i]);
 			Sum = VectorAdd(Sum, Input);
@@ -673,14 +673,14 @@ namespace Audio
 
 	float BufferGetAverageAbsValue(const FAlignedFloatBuffer& Buffer)
 	{
-		checkf(Buffer.Num() % 4 == 0, TEXT("Please use a buffer size that is a multiple of 4."));
+		checkf(Buffer.Num() % AUDIO_NUM_FLOATS_PER_VECTOR_REGISTER == 0, TEXT("Please use a buffer size that is a multiple of %i."), AUDIO_NUM_FLOATS_PER_VECTOR_REGISTER);
 
 		return BufferGetAverageAbsValue(Buffer.GetData(), Buffer.Num());
 	}
 
 	float BufferGetAverageAbsValue(const float* RESTRICT Buffer, int32 NumSamples)
 	{
-		checkf(NumSamples % 4 == 0, TEXT("Please use a buffer size that is a multiple of 4."));
+		checkf(NumSamples % AUDIO_NUM_FLOATS_PER_VECTOR_REGISTER == 0, TEXT("Please use a buffer size that is a multiple of %i."), AUDIO_NUM_FLOATS_PER_VECTOR_REGISTER);
 
 #if !AUDIO_USE_SIMD
 		float Sum = 0.0f;
@@ -692,7 +692,7 @@ namespace Audio
 #else
 		VectorRegister4Float Sum = VectorZero();
 
-		for (int32 i = 0; i < NumSamples; i += 4)
+		for (int32 i = 0; i < NumSamples; i += AUDIO_NUM_FLOATS_PER_VECTOR_REGISTER)
 		{
 			VectorRegister4Float Input = VectorAbs(VectorLoadAligned(&Buffer[i]));
 			Sum = VectorAdd(Sum, Input);
@@ -773,7 +773,7 @@ namespace Audio
 	{
 		const VectorRegister4Float GainVector = VectorLoadFloat2(Gains);
 		
-		for (int32 i = 0; i < NumSamples; i += 4)
+		for (int32 i = 0; i < NumSamples; i += AUDIO_NUM_FLOATS_PER_VECTOR_REGISTER)
 		{
 			VectorRegister4Float Result = VectorLoadAligned(&StereoBuffer[i]);
 			Result = VectorMultiply(Result, GainVector);
@@ -794,7 +794,7 @@ namespace Audio
 		const VectorRegister4Float NumFramesVector = VectorSetFloat1(NumSamples / 4.0f);
 		const VectorRegister4Float GainDeltasVector = VectorDivide(VectorSubtract(DestinationVector, GainVector), NumFramesVector);
 
-		for (int32 i = 0; i < NumSamples; i += 4)
+		for (int32 i = 0; i < NumSamples; i += AUDIO_NUM_FLOATS_PER_VECTOR_REGISTER)
 		{
 			VectorRegister4Float Result = VectorLoadAligned(&StereoBuffer[i]);
 			Result = VectorMultiply(Result, GainVector);
@@ -892,21 +892,21 @@ namespace Audio
 	*/
 	void MixMonoTo2ChannelsFast(const float* RESTRICT MonoBuffer, float* RESTRICT DestinationBuffer, int32 InNumFrames)
 	{
-		checkf(InNumFrames >= 4, TEXT("Buffer must have atleast 4 elements."));
-		checkf(0 == (InNumFrames % 4), TEXT("Buffer length be a multiple of 4."));
-		checkf(IsAligned<const float*>(MonoBuffer, AUDIO_SIMD_FLOAT_ALIGNMENT), TEXT("Memory must be aligned to use vector operations."));
-		checkf(IsAligned<float*>(DestinationBuffer, AUDIO_SIMD_FLOAT_ALIGNMENT), TEXT("Memory must be aligned to use vector operations."));
+		checkf(InNumFrames >= AUDIO_NUM_FLOATS_PER_VECTOR_REGISTER, TEXT("Buffer must have at least %i elements."), AUDIO_NUM_FLOATS_PER_VECTOR_REGISTER);
+		checkf(0 == (InNumFrames % AUDIO_NUM_FLOATS_PER_VECTOR_REGISTER), TEXT("Buffer length be a multiple of %i."), AUDIO_NUM_FLOATS_PER_VECTOR_REGISTER);
+		checkf(IsAligned<const float*>(MonoBuffer, AUDIO_SIMD_BYTE_ALIGNMENT), TEXT("Memory must be aligned to use vector operations."));
+		checkf(IsAligned<float*>(DestinationBuffer, AUDIO_SIMD_BYTE_ALIGNMENT), TEXT("Memory must be aligned to use vector operations."));
 
 		int32 OutPos = 0;
-		for (int32 i = 0; i < InNumFrames; i += 4)
+		for (int32 i = 0; i < InNumFrames; i += AUDIO_NUM_FLOATS_PER_VECTOR_REGISTER)
 		{
 			VectorRegister4Float Input = VectorLoadAligned(&MonoBuffer[i]);
 			VectorRegister4Float Output = VectorSwizzle(Input, 0, 0, 1, 1);
 			VectorStoreAligned(Output, &DestinationBuffer[OutPos]);
-			OutPos += 4;
+			OutPos += AUDIO_NUM_FLOATS_PER_VECTOR_REGISTER;
 			Output = VectorSwizzle(Input, 2, 2, 3, 3);
 			VectorStoreAligned(Output, &DestinationBuffer[OutPos]);
-			OutPos += 4;
+			OutPos += AUDIO_NUM_FLOATS_PER_VECTOR_REGISTER;
 		}
 	}
 
@@ -1018,7 +1018,7 @@ namespace Audio
 	{
 		const VectorRegister4Float GainVector = VectorLoadAligned(Gains);
 
-		for (int32 i = 0; i < NumSamples; i += 4)
+		for (int32 i = 0; i < NumSamples; i += AUDIO_NUM_FLOATS_PER_VECTOR_REGISTER)
 		{
 			VectorRegister4Float Result = VectorLoadAligned(&InterleavedBuffer[i]);
 			Result = VectorMultiply(Result, GainVector);
@@ -1036,10 +1036,10 @@ namespace Audio
 		// Initialize GainVector at StartGains and compute GainDeltasVector:
 		VectorRegister4Float GainVector = VectorLoadAligned(StartGains);
 		const VectorRegister4Float DestinationVector = VectorLoadAligned(EndGains);
-		const VectorRegister4Float NumFramesVector = VectorSetFloat1(NumSamples / 4.0f);
+		const VectorRegister4Float NumFramesVector = VectorSetFloat1(NumSamples / (float)AUDIO_NUM_FLOATS_PER_VECTOR_REGISTER);
 		const VectorRegister4Float GainDeltasVector = VectorDivide(VectorSubtract(DestinationVector, GainVector), NumFramesVector);
 
-		for (int32 i = 0; i < NumSamples; i += 4)
+		for (int32 i = 0; i < NumSamples; i += AUDIO_NUM_FLOATS_PER_VECTOR_REGISTER)
 		{
 			VectorRegister4Float Result = VectorLoadAligned(&InterleavedBuffer[i]);
 			Result = VectorMultiply(Result, GainVector);
@@ -1051,7 +1051,7 @@ namespace Audio
 
 	void MixMonoTo4ChannelsFast(const FAlignedFloatBuffer& MonoBuffer, FAlignedFloatBuffer& DestinationBuffer, const float* RESTRICT Gains)
 	{
-		MixMonoTo4ChannelsFast(MonoBuffer.GetData(), DestinationBuffer.GetData(), DestinationBuffer.Num() / 4, Gains);
+		MixMonoTo4ChannelsFast(MonoBuffer.GetData(), DestinationBuffer.GetData(), DestinationBuffer.Num() / AUDIO_NUM_FLOATS_PER_VECTOR_REGISTER, Gains);
 	}
 
 	/**
@@ -1082,7 +1082,7 @@ namespace Audio
 
 	void MixMonoTo4ChannelsFast(const FAlignedFloatBuffer& MonoBuffer, FAlignedFloatBuffer& DestinationBuffer, const float* RESTRICT StartGains, const float* RESTRICT EndGains)
 	{
-		MixMonoTo4ChannelsFast(MonoBuffer.GetData(), DestinationBuffer.GetData(), DestinationBuffer.Num() / 4, StartGains, EndGains);
+		MixMonoTo4ChannelsFast(MonoBuffer.GetData(), DestinationBuffer.GetData(), DestinationBuffer.Num() / AUDIO_NUM_FLOATS_PER_VECTOR_REGISTER, StartGains, EndGains);
 	}
 
 
@@ -1118,7 +1118,7 @@ namespace Audio
 
 	void Mix2ChannelsTo4ChannelsFast(const FAlignedFloatBuffer& SourceBuffer, FAlignedFloatBuffer& DestinationBuffer, const float* RESTRICT Gains)
 	{
-		Mix2ChannelsTo4ChannelsFast(SourceBuffer.GetData(), DestinationBuffer.GetData(), DestinationBuffer.Num() / 4, Gains);
+		Mix2ChannelsTo4ChannelsFast(SourceBuffer.GetData(), DestinationBuffer.GetData(), DestinationBuffer.Num() / AUDIO_NUM_FLOATS_PER_VECTOR_REGISTER, Gains);
 	}
 
 	/**
@@ -1157,7 +1157,7 @@ namespace Audio
 
 	void Mix2ChannelsTo4ChannelsFast(const FAlignedFloatBuffer& SourceBuffer, FAlignedFloatBuffer& DestinationBuffer, const float* RESTRICT StartGains, const float* RESTRICT EndGains)
 	{
-		Mix2ChannelsTo4ChannelsFast(SourceBuffer.GetData(), DestinationBuffer.GetData(), DestinationBuffer.Num() / 4, StartGains, EndGains);
+		Mix2ChannelsTo4ChannelsFast(SourceBuffer.GetData(), DestinationBuffer.GetData(), DestinationBuffer.Num() / AUDIO_NUM_FLOATS_PER_VECTOR_REGISTER, StartGains, EndGains);
 	}
 
 	/**
@@ -1893,16 +1893,16 @@ namespace Audio
 	/** Interleaves samples from two input buffers */
 	void BufferInterleave2ChannelFast(const float* RESTRICT InBuffer1, const float* RESTRICT InBuffer2, float* RESTRICT OutBuffer, const int32 InNum)
 	{
-		checkf(InNum >= 4, TEXT("Buffer must have atleast 4 elements."));
-		checkf(0 == (InNum % 4), TEXT("Buffer length be a multiple of 4."));
-		checkf(IsAligned<const float*>(InBuffer1, AUDIO_SIMD_FLOAT_ALIGNMENT), TEXT("Memory must be aligned to use vector operations."));
-		checkf(IsAligned<const float*>(InBuffer2, AUDIO_SIMD_FLOAT_ALIGNMENT), TEXT("Memory must be aligned to use vector operations."));
-		checkf(IsAligned<float*>(OutBuffer, AUDIO_SIMD_FLOAT_ALIGNMENT), TEXT("Memory must be aligned to use vector operations."));
+		checkf(InNum >= 4, TEXT("Buffer must have at least %i elements."), AUDIO_NUM_FLOATS_PER_VECTOR_REGISTER);
+		checkf(0 == (InNum % AUDIO_NUM_FLOATS_PER_VECTOR_REGISTER), TEXT("Buffer length be a multiple of %i."), AUDIO_NUM_FLOATS_PER_VECTOR_REGISTER);
+		checkf(IsAligned<const float*>(InBuffer1, AUDIO_SIMD_BYTE_ALIGNMENT), TEXT("Memory must be aligned to use vector operations."));
+		checkf(IsAligned<const float*>(InBuffer2, AUDIO_SIMD_BYTE_ALIGNMENT), TEXT("Memory must be aligned to use vector operations."));
+		checkf(IsAligned<float*>(OutBuffer, AUDIO_SIMD_BYTE_ALIGNMENT), TEXT("Memory must be aligned to use vector operations."));
 
 		const int32 OutNum = 2 * InNum;
 
 		int32 OutPos = 0;
-		for (int32 i = 0; i < InNum; i += 4)
+		for (int32 i = 0; i < InNum; i += AUDIO_NUM_FLOATS_PER_VECTOR_REGISTER)
 		{
 			// Vector1[L0, L1, L2, L3]
 			VectorRegister4Float Vector1 = VectorLoadAligned(&InBuffer1[i]);
@@ -1914,14 +1914,14 @@ namespace Audio
 			// Interleaved[L0, R0, L1, R1]
 			VectorRegister4Float Interleaved = VectorSwizzle(HalfInterleaved, 0, 2, 1, 3);
 			VectorStoreAligned(Interleaved, &OutBuffer[OutPos]);
-			OutPos += 4;
+			OutPos += AUDIO_NUM_FLOATS_PER_VECTOR_REGISTER;
 
 			// HalfInterleaved[L2, L3, R2, R3]
 			HalfInterleaved = VectorShuffle(Vector1, Vector2, 2, 3, 2, 3);
 			// Interleaved[L2, R2, L3, R3]
 			Interleaved = VectorSwizzle(HalfInterleaved, 0, 2, 1, 3);
 			VectorStoreAligned(Interleaved, &OutBuffer[OutPos]);
-			OutPos += 4;
+			OutPos += AUDIO_NUM_FLOATS_PER_VECTOR_REGISTER;
 		}
 	}
 
@@ -1943,11 +1943,11 @@ namespace Audio
 	/** Deinterleaves samples from a 2 channel input buffer */
 	void BufferDeinterleave2ChannelFast(const float* RESTRICT InBuffer, float* RESTRICT OutBuffer1, float* RESTRICT OutBuffer2, const int32 InNumFrames)
 	{
-		checkf(InNumFrames >= 4, TEXT("Buffer must have atleast 4 elements."));
-		checkf(0 == (InNumFrames % 4), TEXT("Buffer length be a multiple of 4."));
-		checkf(IsAligned<const float*>(InBuffer, AUDIO_SIMD_FLOAT_ALIGNMENT), TEXT("Memory must be aligned to use vector operations."));
-		checkf(IsAligned<const float*>(OutBuffer1, 4), TEXT("Memory must be aligned to use vector operations."));
-		checkf(IsAligned<float*>(OutBuffer2, 4), TEXT("Memory must be aligned to use vector operations."));
+		checkf(InNumFrames >= AUDIO_NUM_FLOATS_PER_VECTOR_REGISTER, TEXT("Buffer must have at least %i elements."), AUDIO_NUM_FLOATS_PER_VECTOR_REGISTER);
+		checkf(0 == (InNumFrames % AUDIO_NUM_FLOATS_PER_VECTOR_REGISTER), TEXT("Buffer length be a multiple of %i."), AUDIO_NUM_FLOATS_PER_VECTOR_REGISTER);
+		checkf(IsAligned<const float*>(InBuffer, AUDIO_SIMD_BYTE_ALIGNMENT), TEXT("Memory must be aligned to use vector operations."));
+		checkf(IsAligned<const float*>(OutBuffer1, AUDIO_SIMD_BYTE_ALIGNMENT), TEXT("Memory must be aligned to use vector operations."));
+		checkf(IsAligned<float*>(OutBuffer2, AUDIO_SIMD_BYTE_ALIGNMENT), TEXT("Memory must be aligned to use vector operations."));
 
 		int32 InNum = InNumFrames * 2;
 		int32 OutPos = 0;
@@ -1965,7 +1965,7 @@ namespace Audio
 			OutVector = VectorShuffle(InVector1, InVector2, 1, 3, 1, 3);
 			VectorStoreAligned(OutVector, &OutBuffer2[OutPos]);
 
-			OutPos += 4;
+			OutPos += AUDIO_NUM_FLOATS_PER_VECTOR_REGISTER;
 		}
 	}
 
@@ -1984,10 +1984,10 @@ namespace Audio
 	/** Sums 2 channel interleaved input samples. OutSamples[n] = InSamples[2n] + InSamples[2n + 1] */
 	void BufferSum2ChannelToMonoFast(const float* RESTRICT InSamples, float* RESTRICT OutSamples, const int32 InNumFrames)
 	{
-		checkf(InNumFrames >= 4, TEXT("Buffer must have atleast 4 elements."));
-		checkf(0 == (InNumFrames % 4), TEXT("Buffer length be a multiple of 4."));
-		checkf(IsAligned<const float*>(InSamples, AUDIO_SIMD_FLOAT_ALIGNMENT), TEXT("Memory must be aligned to use vector operations."));
-		checkf(IsAligned<float*>(OutSamples, 4), TEXT("Memory must be aligned to use vector operations."));
+		checkf(InNumFrames >= AUDIO_NUM_FLOATS_PER_VECTOR_REGISTER, TEXT("Buffer must have at least %i elements."), AUDIO_NUM_FLOATS_PER_VECTOR_REGISTER);
+		checkf(0 == (InNumFrames % AUDIO_NUM_FLOATS_PER_VECTOR_REGISTER), TEXT("Buffer length be a multiple of %i."), AUDIO_NUM_FLOATS_PER_VECTOR_REGISTER);
+		checkf(IsAligned<const float*>(InSamples, AUDIO_SIMD_BYTE_ALIGNMENT), TEXT("Memory must be aligned to use vector operations."));
+		checkf(IsAligned<float*>(OutSamples, AUDIO_SIMD_BYTE_ALIGNMENT), TEXT("Memory must be aligned to use vector operations."));
 
 		const int32 InNum = InNumFrames * 2;
 		int32 OutPos = 0;
@@ -2010,7 +2010,7 @@ namespace Audio
 			VectorRegister4Float Out = VectorAdd(Channel0, Channel1);
 
 			VectorStoreAligned(Out, &OutSamples[OutPos]);
-			OutPos += 4;
+			OutPos += AUDIO_NUM_FLOATS_PER_VECTOR_REGISTER;
 		}
 	}
 
@@ -2028,17 +2028,17 @@ namespace Audio
 
 	void BufferComplexToPowerFast(const float* RESTRICT InRealSamples, const float* RESTRICT InImaginarySamples, float* RESTRICT OutPowerSamples, const int32 InNum)
 	{
-		checkf(InNum >= 4, TEXT("Buffer must have atleast 4 elements."));
-		checkf(0 == (InNum % 4), TEXT("Buffer length be a multiple of 4."));
+		checkf(InNum >= 4, TEXT("Buffer must have at least %i elements."), AUDIO_NUM_FLOATS_PER_VECTOR_REGISTER);
+		checkf(0 == (InNum % AUDIO_NUM_FLOATS_PER_VECTOR_REGISTER), TEXT("Buffer length be a multiple of %i."), AUDIO_NUM_FLOATS_PER_VECTOR_REGISTER);
 
-		checkf(IsAligned<const float*>(InRealSamples, AUDIO_SIMD_FLOAT_ALIGNMENT), TEXT("Memory must be aligned to use vector operations."));
-		checkf(IsAligned<const float*>(InImaginarySamples, AUDIO_SIMD_FLOAT_ALIGNMENT), TEXT("Memory must be aligned to use vector operations."));
-		checkf(IsAligned<float*>(OutPowerSamples, AUDIO_SIMD_FLOAT_ALIGNMENT), TEXT("Memory must be aligned to use vector operations."));
+		checkf(IsAligned<const float*>(InRealSamples, AUDIO_SIMD_BYTE_ALIGNMENT), TEXT("Memory must be aligned to use vector operations."));
+		checkf(IsAligned<const float*>(InImaginarySamples, AUDIO_SIMD_BYTE_ALIGNMENT), TEXT("Memory must be aligned to use vector operations."));
+		checkf(IsAligned<float*>(OutPowerSamples, AUDIO_SIMD_BYTE_ALIGNMENT), TEXT("Memory must be aligned to use vector operations."));
 
 		RestrictedPtrAliasCheck(InRealSamples, OutPowerSamples, InNum);
 		RestrictedPtrAliasCheck(InImaginarySamples, OutPowerSamples, InNum);
 
-		for (int32 i = 0; i < InNum; i += 4)
+		for (int32 i = 0; i < InNum; i += AUDIO_NUM_FLOATS_PER_VECTOR_REGISTER)
 		{
 			VectorRegister4Float VInReal = VectorLoadAligned(&InRealSamples[i]);
 			VectorRegister4Float VInRealSquared = VectorMultiply(VInReal, VInReal);
diff --git a/Engine/Source/Runtime/SignalProcessing/Private/DynamicDelayAPF.cpp b/Engine/Source/Runtime/SignalProcessing/Private/DynamicDelayAPF.cpp
--- a/Engine/Source/Runtime/SignalProcessing/Private/DynamicDelayAPF.cpp
+++ b/Engine/Source/Runtime/SignalProcessing/Private/DynamicDelayAPF.cpp
@@ -22,7 +22,7 @@ FDynamicDelayAPF::FDynamicDelayAPF(float InG, int32 InMinDelay, int32 InMaxDelay
 	{
 		NumInternalBufferSamples = MaxBufferSamples;
 		// Block length must be divisible by simd alignment to support simd operations.
-		NumInternalBufferSamples -= (NumInternalBufferSamples % AUDIO_SIMD_FLOAT_ALIGNMENT);
+		NumInternalBufferSamples -= (NumInternalBufferSamples % AUDIO_NUM_FLOATS_PER_VECTOR_REGISTER);
 	}
 
 	checkf(NumInternalBufferSamples > 0, TEXT("Invalid internal buffer length"));
diff --git a/Engine/Source/Runtime/SignalProcessing/Private/FeedbackDelayNetwork.cpp b/Engine/Source/Runtime/SignalProcessing/Private/FeedbackDelayNetwork.cpp
--- a/Engine/Source/Runtime/SignalProcessing/Private/FeedbackDelayNetwork.cpp
+++ b/Engine/Source/Runtime/SignalProcessing/Private/FeedbackDelayNetwork.cpp
@@ -43,7 +43,7 @@ FFeedbackDelayNetwork::FFeedbackDelayNetwork(int32 InMaxNumInternalBufferSamples
 
 	NumInternalBufferSamples = FMath::Min(NumInternalBufferSamples, InMaxNumInternalBufferSamples);
 
-	NumInternalBufferSamples -= (NumInternalBufferSamples % AUDIO_SIMD_FLOAT_ALIGNMENT);
+	NumInternalBufferSamples -= (NumInternalBufferSamples % AUDIO_NUM_FLOATS_PER_VECTOR_REGISTER);
 
 	if (NumInternalBufferSamples <= 0)
 	{
diff --git a/Engine/Source/Runtime/SignalProcessing/Private/IntegerDelay.cpp b/Engine/Source/Runtime/SignalProcessing/Private/IntegerDelay.cpp
--- a/Engine/Source/Runtime/SignalProcessing/Private/IntegerDelay.cpp
+++ b/Engine/Source/Runtime/SignalProcessing/Private/IntegerDelay.cpp
@@ -45,7 +45,7 @@ void FIntegerDelay::SetDelayLengthSamples(int32 InNumDelaySamples)
 	// and abide by alignment rules of FAlignedBlockBuffer
 	NumDelayLineOffsetSamples = MaxNumDelaySamples - NumDelaySamples;
 	NumBufferOffsetSamples = 0;
-	while (0 != (NumDelayLineOffsetSamples % AUDIO_SIMD_FLOAT_ALIGNMENT))
+	while (0 != (NumDelayLineOffsetSamples % AUDIO_NUM_FLOATS_PER_VECTOR_REGISTER))
 	{
 		NumDelayLineOffsetSamples--;
 		NumBufferOffsetSamples++;
diff --git a/Engine/Source/Runtime/SignalProcessing/Private/LateReflectionsFast.cpp b/Engine/Source/Runtime/SignalProcessing/Private/LateReflectionsFast.cpp
--- a/Engine/Source/Runtime/SignalProcessing/Private/LateReflectionsFast.cpp
+++ b/Engine/Source/Runtime/SignalProcessing/Private/LateReflectionsFast.cpp
@@ -120,7 +120,7 @@ FLateReflectionsPlate::FLateReflectionsPlate(
 	if (NumInternalBufferSamples > PlateDelays.NumSamplesDelayI) {
 		NumInternalBufferSamples = PlateDelays.NumSamplesDelayI;
 	}
-	NumInternalBufferSamples -= (NumInternalBufferSamples % AUDIO_SIMD_FLOAT_ALIGNMENT);
+	NumInternalBufferSamples -= (NumInternalBufferSamples % AUDIO_NUM_FLOATS_PER_VECTOR_REGISTER);
 
 	checkf(NumInternalBufferSamples > 0, TEXT("InDelays.DelaySample2C too small"));
 
diff --git a/Engine/Source/Runtime/SignalProcessing/Private/LinearInterpFractionalDelay.cpp b/Engine/Source/Runtime/SignalProcessing/Private/LinearInterpFractionalDelay.cpp
--- a/Engine/Source/Runtime/SignalProcessing/Private/LinearInterpFractionalDelay.cpp
+++ b/Engine/Source/Runtime/SignalProcessing/Private/LinearInterpFractionalDelay.cpp
@@ -18,13 +18,13 @@ FLinearInterpFractionalDelay::FLinearInterpFractionalDelay(int32 InMaxDelay, int
 		MaxDelay = 1;
 	}
 
-	while (0 != (NumInternalBufferSamples % AUDIO_SIMD_FLOAT_ALIGNMENT))
+	while (0 != (NumInternalBufferSamples % AUDIO_NUM_FLOATS_PER_VECTOR_REGISTER))
 	{
 		NumInternalBufferSamples--;
 	}
 	if (NumInternalBufferSamples < 1)
 	{
-		NumInternalBufferSamples = AUDIO_SIMD_FLOAT_ALIGNMENT;
+		NumInternalBufferSamples = AUDIO_NUM_FLOATS_PER_VECTOR_REGISTER;
 	}
 
 	// Allocate and prepare delay line for maximum delay.
@@ -99,9 +99,9 @@ void FLinearInterpFractionalDelay::ProcessAudio(const FAlignedFloatBuffer& InSam
 void FLinearInterpFractionalDelay::ProcessAudioBlock(const float* InSamples, const float* InDelays, const int32 InNum, float* OutSamples)
 {
 	checkf(0 == (InNum % 4), TEXT("Array length must be multiple of 4"));
-	checkf(IsAligned<const float*>(InSamples, AUDIO_SIMD_FLOAT_ALIGNMENT), TEXT("Memory must be aligned to use vector operations."));
-	checkf(IsAligned<const float*>(InDelays, AUDIO_SIMD_FLOAT_ALIGNMENT), TEXT("Memory must be aligned to use vector operations."));
-	checkf(IsAligned<float*>(OutSamples, AUDIO_SIMD_FLOAT_ALIGNMENT), TEXT("Memory must be aligned to use vector operations."));
+	checkf(IsAligned<const float*>(InSamples, AUDIO_SIMD_BYTE_ALIGNMENT), TEXT("Memory must be aligned to use vector operations."));
+	checkf(IsAligned<const float*>(InDelays, AUDIO_SIMD_BYTE_ALIGNMENT), TEXT("Memory must be aligned to use vector operations."));
+	checkf(IsAligned<float*>(OutSamples, AUDIO_SIMD_BYTE_ALIGNMENT), TEXT("Memory must be aligned to use vector operations."));
 
 
 	// Update delay line.
diff --git a/Engine/Source/Runtime/SignalProcessing/Private/LongDelayAPF.cpp b/Engine/Source/Runtime/SignalProcessing/Private/LongDelayAPF.cpp
--- a/Engine/Source/Runtime/SignalProcessing/Private/LongDelayAPF.cpp
+++ b/Engine/Source/Runtime/SignalProcessing/Private/LongDelayAPF.cpp
@@ -17,7 +17,7 @@ FLongDelayAPF::FLongDelayAPF(float InG, int32 InNumDelaySamples, int32 InMaxNumI
 		NumInternalBufferSamples = NumDelaySamples;
 
 		// Block size must be divisible by simd alignment to support simd operations.
-		NumInternalBufferSamples -= (NumInternalBufferSamples % AUDIO_SIMD_FLOAT_ALIGNMENT);
+		NumInternalBufferSamples -= (NumInternalBufferSamples % AUDIO_NUM_FLOATS_PER_VECTOR_REGISTER);
 	}
 
 	checkf(NumInternalBufferSamples > 1, TEXT("Invalid internal buffer length"));
@@ -105,7 +105,7 @@ void FLongDelayAPF::ProcessAudioBlock(const float* InSamples, const float* InDel
 	//WeightedSumBuffer(InDelaySamples, G, InSamples, OutDelaySamples, InNum);
 	
 	
-	int32 NumToSIMD = InNum - (InNum % AUDIO_SIMD_FLOAT_ALIGNMENT);
+	int32 NumToSIMD = InNum - (InNum % AUDIO_NUM_FLOATS_PER_VECTOR_REGISTER);
 
 	VectorRegister VG = MakeVectorRegister(G, G, G, G);
 	VectorRegister VNG = MakeVectorRegister(-G, -G, -G, -G);
diff --git a/Engine/Source/Runtime/SignalProcessing/Private/VectorFFT.cpp b/Engine/Source/Runtime/SignalProcessing/Private/VectorFFT.cpp
--- a/Engine/Source/Runtime/SignalProcessing/Private/VectorFFT.cpp
+++ b/Engine/Source/Runtime/SignalProcessing/Private/VectorFFT.cpp
@@ -52,8 +52,8 @@ namespace Audio
 			// @param OutComplex - Interleaved complex data with (2 * FFTSize) num floats.
 			void ForwardComplexToComplex(const float* RESTRICT InComplex, float* RESTRICT OutComplex)
 			{
-				checkf(IsAligned<const float*>(InComplex, AUDIO_SIMD_FLOAT_ALIGNMENT), TEXT("Memory must be aligned.")); 
-				checkf(IsAligned<float*>(OutComplex, AUDIO_SIMD_FLOAT_ALIGNMENT), TEXT("Memory must be aligned.")); 
+				checkf(IsAligned<const float*>(InComplex, AUDIO_SIMD_BYTE_ALIGNMENT), TEXT("Memory must be aligned."));
+				checkf(IsAligned<float*>(OutComplex, AUDIO_SIMD_BYTE_ALIGNMENT), TEXT("Memory must be aligned."));
 
 				// To perform FFT, must complete Log2FFTSize stages.  Each radix pass performs 2^m stages
 				// where 2^m is the radix number. So a radix-4 stage is radix-2^m or radix-2^2. Hence radix
@@ -108,8 +108,8 @@ namespace Audio
 
 			void InverseComplexToComplex(const float* RESTRICT InComplex, float* RESTRICT OutComplex)
 			{
-				checkf(IsAligned<const float*>(InComplex, AUDIO_SIMD_FLOAT_ALIGNMENT), TEXT("Memory must be aligned.")); 
-				checkf(IsAligned<float*>(OutComplex, AUDIO_SIMD_FLOAT_ALIGNMENT), TEXT("Memory must be aligned.")); 
+				checkf(IsAligned<const float*>(InComplex, AUDIO_SIMD_BYTE_ALIGNMENT), TEXT("Memory must be aligned."));
+				checkf(IsAligned<float*>(OutComplex, AUDIO_SIMD_BYTE_ALIGNMENT), TEXT("Memory must be aligned."));
 
 				// Perform inverse FFT by complex conjugating the input and output.
 				
@@ -1152,8 +1152,8 @@ namespace Audio
 
 	void FVectorRealToComplexFFT::ForwardRealToComplex(const float* RESTRICT InReal, float* RESTRICT OutComplex)
 	{
-		checkf(IsAligned<const float*>(InReal, AUDIO_SIMD_FLOAT_ALIGNMENT), TEXT("Memory must be aligned.")); 
-		checkf(IsAligned<float*>(OutComplex, AUDIO_SIMD_FLOAT_ALIGNMENT), TEXT("Memory must be aligned.")); 
+		checkf(IsAligned<const float*>(InReal, AUDIO_SIMD_BYTE_ALIGNMENT), TEXT("Memory must be aligned."));
+		checkf(IsAligned<float*>(OutComplex, AUDIO_SIMD_BYTE_ALIGNMENT), TEXT("Memory must be aligned."));
 
 		// Performs a N sized real-to-complex FFT using an N/2 complex-to-complex FFT.
 		float* WorkData = WorkBuffer.GetData();
@@ -1197,8 +1197,8 @@ namespace Audio
 
 	void FVectorRealToComplexFFT::InverseComplexToReal(const float* RESTRICT InComplex, float* RESTRICT OutReal)
 	{
-		checkf(IsAligned<const float*>(InComplex, AUDIO_SIMD_FLOAT_ALIGNMENT), TEXT("Memory must be aligned.")); 
-		checkf(IsAligned<float*>(OutReal, AUDIO_SIMD_FLOAT_ALIGNMENT), TEXT("Memory must be aligned.")); 
+		checkf(IsAligned<const float*>(InComplex, AUDIO_SIMD_BYTE_ALIGNMENT), TEXT("Memory must be aligned."));
+		checkf(IsAligned<float*>(OutReal, AUDIO_SIMD_BYTE_ALIGNMENT), TEXT("Memory must be aligned."));
 
 		// Performs a N sized complex-to-real FFT using an N/2 complex-to-complex FFT.
 
diff --git a/Engine/Source/Runtime/SignalProcessing/Public/DSP/BufferVectorOperations.h b/Engine/Source/Runtime/SignalProcessing/Public/DSP/BufferVectorOperations.h
--- a/Engine/Source/Runtime/SignalProcessing/Public/DSP/BufferVectorOperations.h
+++ b/Engine/Source/Runtime/SignalProcessing/Public/DSP/BufferVectorOperations.h
@@ -12,7 +12,7 @@
 #endif
 
 #define AUDIO_SIMD_BYTE_ALIGNMENT (16)
-#define AUDIO_SIMD_FLOAT_ALIGNMENT (4)
+#define AUDIO_NUM_FLOATS_PER_VECTOR_REGISTER (4)
 
 namespace Audio
 {
