#ue5 - LWC: move more audio DSP types to just use the float version in all builds.

diff --git a/Engine/Source/Runtime/SignalProcessing/Private/FloatArrayMath.cpp b/Engine/Source/Runtime/SignalProcessing/Private/FloatArrayMath.cpp
--- a/Engine/Source/Runtime/SignalProcessing/Private/FloatArrayMath.cpp
+++ b/Engine/Source/Runtime/SignalProcessing/Private/FloatArrayMath.cpp
@@ -37,11 +37,11 @@ namespace Audio
 
 		if (NumToSimd)
 		{
-			VectorRegister Total = VectorSetFloat1(0.f);
+			VectorRegister4Float Total = VectorSetFloat1(0.f);
 
 			for (int32 i = 0; i < NumToSimd; i += 4)
 			{
-				VectorRegister VectorData = VectorLoadAligned(&InData[i]);
+				VectorRegister4Float VectorData = VectorLoadAligned(&InData[i]);
 				Total = VectorAdd(Total, VectorData);
 			}
 
@@ -425,18 +425,18 @@ namespace Audio
 
 		if (NumToSimd)
 		{
-			const VectorRegister RealSignFlip = MakeVectorRegister(-1.f, 1.f, -1.f, 1.f);
+			const VectorRegister4Float RealSignFlip = MakeVectorRegister(-1.f, 1.f, -1.f, 1.f);
 
 			for (int32 i = 0; i < NumToSimd; i += 4)
 			{
-				VectorRegister VectorData1 = VectorLoadAligned(&InData1[i]);
-				VectorRegister VectorData2 = VectorLoadAligned(&InData2[i]);
+				VectorRegister4Float VectorData1 = VectorLoadAligned(&InData1[i]);
+				VectorRegister4Float VectorData2 = VectorLoadAligned(&InData2[i]);
 
-				VectorRegister VectorData1Real = VectorSwizzle(VectorData1, 0, 0, 2, 2);
-				VectorRegister VectorData1Imag = VectorSwizzle(VectorData1, 1, 1, 3, 3);
-				VectorRegister VectorData2Swizzle = VectorSwizzle(VectorData2, 1, 0, 3, 2);
+				VectorRegister4Float VectorData1Real = VectorSwizzle(VectorData1, 0, 0, 2, 2);
+				VectorRegister4Float VectorData1Imag = VectorSwizzle(VectorData1, 1, 1, 3, 3);
+				VectorRegister4Float VectorData2Swizzle = VectorSwizzle(VectorData2, 1, 0, 3, 2);
 
-				VectorRegister Result = VectorMultiply(VectorData1Imag, VectorData2Swizzle);
+				VectorRegister4Float Result = VectorMultiply(VectorData1Imag, VectorData2Swizzle);
 				Result = VectorMultiply(Result, RealSignFlip);
 				Result = VectorMultiplyAdd(VectorData1Real, VectorData2, Result);
 
@@ -513,10 +513,10 @@ namespace Audio
 
 		for (int32 i = 0; i < NumToSimd; i += 4)
 		{
-			VectorRegister VectorData = VectorLoadAligned(&InData[i]);
-			VectorRegister VectorAccumData = VectorLoadAligned(&InAccumulateData[i]);
+			VectorRegister4Float VectorData = VectorLoadAligned(&InData[i]);
+			VectorRegister4Float VectorAccumData = VectorLoadAligned(&InAccumulateData[i]);
 
-			VectorRegister VectorOut = VectorAdd(VectorData, VectorAccumData);
+			VectorRegister4Float VectorOut = VectorAdd(VectorData, VectorAccumData);
 			VectorStoreAligned(VectorOut, &InAccumulateData[i]);
 		}
 
@@ -599,15 +599,15 @@ namespace Audio
 		const float Delta = (InEndMultiplier - InStartMultiplier) / FMath::Max(1.f, static_cast<float>(Num - 1));
 
 		const float FourByDelta = 4.f * Delta;
-		VectorRegister VectorDelta = MakeVectorRegister(FourByDelta, FourByDelta, FourByDelta, FourByDelta);
-		VectorRegister VectorMultiplier = MakeVectorRegister(InStartMultiplier, InStartMultiplier + Delta, InStartMultiplier + 2.f * Delta, InStartMultiplier + 3.f * Delta);
+		VectorRegister4Float VectorDelta = MakeVectorRegister(FourByDelta, FourByDelta, FourByDelta, FourByDelta);
+		VectorRegister4Float VectorMultiplier = MakeVectorRegister(InStartMultiplier, InStartMultiplier + Delta, InStartMultiplier + 2.f * Delta, InStartMultiplier + 3.f * Delta);
 
 		for (int32 i = 0; i < NumToSimd; i += 4)
 		{
-			VectorRegister VectorData = VectorLoadAligned(&InData[i]);
-			VectorRegister VectorAccumData = VectorLoadAligned(&InAccumulateData[i]);
+			VectorRegister4Float VectorData = VectorLoadAligned(&InData[i]);
+			VectorRegister4Float VectorAccumData = VectorLoadAligned(&InAccumulateData[i]);
 
-			VectorRegister VectorOut = VectorMultiplyAdd(VectorData, VectorMultiplier, VectorAccumData);
+			VectorRegister4Float VectorOut = VectorMultiplyAdd(VectorData, VectorMultiplier, VectorAccumData);
 			VectorMultiplier = VectorAdd(VectorMultiplier, VectorDelta);
 
 			VectorStoreAligned(VectorOut, &InAccumulateData[i]);
@@ -640,11 +640,11 @@ namespace Audio
 
 		float* InData = InValues.GetData();
 
-		const VectorRegister VectorSubtrahend = VectorSetFloat1(InSubtrahend);
+		const VectorRegister4Float VectorSubtrahend = VectorSetFloat1(InSubtrahend);
 
 		for (int32 i = 0; i < NumToSimd; i += 4)
 		{
-			VectorRegister VectorData = VectorLoadAligned(&InData[i]);
+			VectorRegister4Float VectorData = VectorLoadAligned(&InData[i]);
 			VectorData = VectorSubtract(VectorData, VectorSubtrahend);
 			VectorStoreAligned(VectorData, &InData[i]);
 		}
@@ -720,11 +720,11 @@ namespace Audio
 		const float* InData = InValues.GetData();
 		float* OutData = OutValues.GetData();
 
-		const VectorRegister ConjugateMult = MakeVectorRegister(1.f, -1.f, 1.f, -1.f);
+		const VectorRegister4Float ConjugateMult = MakeVectorRegister(1.f, -1.f, 1.f, -1.f);
 
 		for (int32 i = 0; i < NumToSimd; i += 4)
 		{
-			VectorRegister VectorData = VectorLoadAligned(&InData[i]);
+			VectorRegister4Float VectorData = VectorLoadAligned(&InData[i]);
 			
 			VectorData = VectorMultiply(VectorData, ConjugateMult);
 
@@ -762,11 +762,11 @@ namespace Audio
 
 		float* InData = InValues.GetData();
 
-		const VectorRegister ConjugateMult = MakeVectorRegister(1.f, -1.f, 1.f, -1.f);
+		const VectorRegister4Float ConjugateMult = MakeVectorRegister(1.f, -1.f, 1.f, -1.f);
 
 		for (int32 i = 0; i < NumToSimd; i += 4)
 		{
-			VectorRegister VectorData = VectorLoadAligned(&InData[i]);
+			VectorRegister4Float VectorData = VectorLoadAligned(&InData[i]);
 			
 			VectorData = VectorMultiply(VectorData, ConjugateMult);
 
@@ -806,12 +806,12 @@ namespace Audio
 		const float Scale = 20.f / MathIntrinsics::Loge10;
 		const float Minimum = FMath::Exp(InMinimumDb * MathIntrinsics::Loge10 / 20.f);
 
-		const VectorRegister VectorScale = VectorSetFloat1(Scale);
-		const VectorRegister VectorMinimum = VectorSetFloat1(Minimum);
+		const VectorRegister4Float VectorScale = VectorSetFloat1(Scale);
+		const VectorRegister4Float VectorMinimum = VectorSetFloat1(Minimum);
 
 		for (int32 i = 0; i < NumToSimd; i += 4)
 		{
-			VectorRegister VectorData = VectorLoadAligned(&InData[i]);
+			VectorRegister4Float VectorData = VectorLoadAligned(&InData[i]);
 			
 			VectorData = VectorMax(VectorData, VectorMinimum);
 			VectorData = VectorLog(VectorData);
@@ -852,12 +852,12 @@ namespace Audio
 		const float Scale = 10.f / MathIntrinsics::Loge10;
 		const float Minimum = FMath::Exp(InMinimumDb * MathIntrinsics::Loge10 / 10.f);
 
-		const VectorRegister VectorMinimum = VectorSetFloat1(Minimum);
-		const VectorRegister VectorScale = VectorSetFloat1(Scale);
+		const VectorRegister4Float VectorMinimum = VectorSetFloat1(Minimum);
+		const VectorRegister4Float VectorScale = VectorSetFloat1(Scale);
 
 		for (int32 i = 0; i < NumToSimd; i += 4)
 		{
-			VectorRegister VectorData = VectorLoadAligned(&InData[i]);
+			VectorRegister4Float VectorData = VectorLoadAligned(&InData[i]);
 
 			VectorData = VectorMax(VectorData, VectorMinimum);
 			VectorData = VectorLog(VectorData);
@@ -908,16 +908,16 @@ namespace Audio
 
 		for (int32 i = 0; i < NumToSimd; i += 4)
 		{
-			VectorRegister VectorComplex1 = VectorLoadAligned(&InComplexData[2 * i]);
-			VectorRegister VectorSquared1 = VectorMultiply (VectorComplex1, VectorComplex1);
+			VectorRegister4Float VectorComplex1 = VectorLoadAligned(&InComplexData[2 * i]);
+			VectorRegister4Float VectorSquared1 = VectorMultiply (VectorComplex1, VectorComplex1);
 
-			VectorRegister VectorComplex2 = VectorLoadAligned(&InComplexData[(2 * i) + 4]);
-			VectorRegister VectorSquared2 = VectorMultiply (VectorComplex2, VectorComplex2);
+			VectorRegister4Float VectorComplex2 = VectorLoadAligned(&InComplexData[(2 * i) + 4]);
+			VectorRegister4Float VectorSquared2 = VectorMultiply (VectorComplex2, VectorComplex2);
 
-			VectorRegister VectorSquareReal = VectorShuffle(VectorSquared1, VectorSquared2, 0, 2, 0, 2);
-			VectorRegister VectorSquareImag = VectorShuffle(VectorSquared1, VectorSquared2, 1, 3, 1, 3);
+			VectorRegister4Float VectorSquareReal = VectorShuffle(VectorSquared1, VectorSquared2, 0, 2, 0, 2);
+			VectorRegister4Float VectorSquareImag = VectorShuffle(VectorSquared1, VectorSquared2, 1, 3, 1, 3);
 
-			VectorRegister VectorOut = VectorAdd(VectorSquareReal, VectorSquareImag);
+			VectorRegister4Float VectorOut = VectorAdd(VectorSquareReal, VectorSquareImag);
 			
 			VectorStoreAligned(VectorOut, &OutPowerData[i]);
 		}
