Fix for "Pink Noise Generator not generating full frequency spectrum"

diff --git a/Engine/Source/Runtime/SignalProcessing/Private/Noise.cpp b/Engine/Source/Runtime/SignalProcessing/Private/Noise.cpp
--- a/Engine/Source/Runtime/SignalProcessing/Private/Noise.cpp
+++ b/Engine/Source/Runtime/SignalProcessing/Private/Noise.cpp
@@ -14,13 +14,11 @@ namespace Audio
 		: FWhiteNoise{ static_cast<int32>(FPlatformTime::Cycles()) }
 	{}
 
-	static const float DefaultFilterGain = Audio::ConvertToLinear(-3.f);
-
 	FPinkNoise::FPinkNoise(int32 InRandomSeed)
 		: Noise{ InRandomSeed }
 		, X_Z{ 0,0,0,0 }
 		, Y_Z{ 0,0.0,0 }
-		, A0{ DefaultFilterGain }
+		, A0{ 1.0f }	// Not used.
 	{
 		static_assert(UE_ARRAY_COUNT(X_Z) == 4, "sizeof(X_Z)==4");
 		static_assert(UE_ARRAY_COUNT(Y_Z) == 4, "sizeof(Y_Z)==4");
@@ -34,24 +32,20 @@ namespace Audio
 	{
 		// Filter Coefficients based on:
 		// https://ccrma.stanford.edu/~jos/sasp/Example_Synthesis_1_F_Noise.html
-		static constexpr float A[4] { 1.0f, -2.494956002f,2.017265875f, -0.522189400f };
+		static constexpr float A[3] { -2.494956002f,2.017265875f, -0.522189400f };
 		static constexpr float B[4] { 0.049922035f,-0.095993537f,0.050612699f,-0.004408786f};	
 
-		static_assert(UE_ARRAY_COUNT(A) == UE_ARRAY_COUNT(X_Z), "A Coefficients and X need to be the same size");
-		static_assert(UE_ARRAY_COUNT(B) == UE_ARRAY_COUNT(Y_Z), "B Coefficients and Y need to be the same size");
-
 		X_Z[0] = Noise.Generate(); // Xn
 
-		float Yn =
-				A0 * X_Z[0]		// a0 * (x) (inject non-const filter-gain A0 in here).
-			+ A[1] * X_Z[1]		// a1 * (x-1)
-			+ A[2] * X_Z[2]		// a2 * (x-2)
-			+ A[3] * X_Z[3]		// a3 * (x-3) 
+		float Yn = 
+			  X_Z[0]*B[0] 
+			+ X_Z[1]*B[1] 
+			+ X_Z[2]*B[2] 
+			+ X_Z[3]*B[3]
 
-			- B[0] * Y_Z[0]		// b1 * (y-1)
-			- B[1] * Y_Z[1]		// b2 * (y-2)
-			- B[2] * Y_Z[2]		// b3 * (y-3)
-			- B[3] * Y_Z[3];	// b4 * (y-4)
+			- Y_Z[0]*A[0] 
+			- Y_Z[1]*A[1] 
+			- Y_Z[2]*A[2];
 
 		// Shuffle feed-forward state by one.
 		X_Z[3] = X_Z[2];
