Fixed a bug where the Envelope Follower wasn't properly calculating MS and RMS envelope values

diff --git a/Engine/Source/Runtime/SignalProcessing/Private/EnvelopeFollower.cpp b/Engine/Source/Runtime/SignalProcessing/Private/EnvelopeFollower.cpp
--- a/Engine/Source/Runtime/SignalProcessing/Private/EnvelopeFollower.cpp
+++ b/Engine/Source/Runtime/SignalProcessing/Private/EnvelopeFollower.cpp
@@ -12,7 +12,10 @@ namespace Audio
 	static const float DigitalTimeConstant = 4.60517019f;
 
 	FEnvelopeFollower::FEnvelopeFollower()
-		: EnvMode(EPeakMode::Peak)	
+		: EnvMode(EPeakMode::Peak)
+		, MeanWindowSize(DefaultWindowSize)
+		, MeanHopSize(DefaultHopSize)
+		, SumBuffer(DefaultWindowSize, DefaultHopSize)
 		, SampleRate(44100.0f)
 		, AttackTimeMsec(0.0f)
 		, AttackTimeSamples(0.0f)
@@ -23,22 +26,36 @@ namespace Audio
 	{
 	}
 
-	FEnvelopeFollower::FEnvelopeFollower(const float InSampleRate, const float InAttackTimeMsec, const float InReleaseTimeMSec, const EPeakMode::Type InMode, const bool bInIsAnalg)
+	FEnvelopeFollower::FEnvelopeFollower(const float InSampleRate,
+		const float InAttackTimeMsec,
+		const float InReleaseTimeMSec,
+		const EPeakMode::Type InMode,
+		const int32 InWindowSizeForMean,
+		const int32 InHopSizeForMean,
+		const bool bInIsAnalog) : SumBuffer(InWindowSizeForMean, InHopSizeForMean)
 	{
-		Init(InSampleRate, InAttackTimeMsec, InReleaseTimeMSec, InMode, bInIsAnalg);
+		Init(InSampleRate, InAttackTimeMsec, InReleaseTimeMSec, InMode, InWindowSizeForMean, InHopSizeForMean, bInIsAnalog);
 	}
 
 	FEnvelopeFollower::~FEnvelopeFollower()
 	{
 	}
 
-	void FEnvelopeFollower::Init(const float InSampleRate, const float InAttackTimeMsec, const float InReleaseTimeMSec, const EPeakMode::Type InMode, const bool bInIsAnalog)
+	void FEnvelopeFollower::Init(const float InSampleRate, 
+		const float InAttackTimeMsec, 
+		const float InReleaseTimeMSec, 
+		const EPeakMode::Type InMode, 
+		const int32 InWindowSizeForMean,
+		const int32 InHopSizeForMean,
+		const bool bInIsAnalog)
 	{
 		SampleRate = InSampleRate;
 
 		bIsAnalog = bInIsAnalog;
 		EnvMode = InMode;
 
+		SumBuffer = TSlidingBuffer<float>(InWindowSizeForMean, InHopSizeForMean);
+
 		// Set the attack and release times using the default values
 		SetAttackTime(InAttackTimeMsec);
 		SetReleaseTime(InReleaseTimeMSec);
@@ -77,30 +94,77 @@ namespace Audio
 
 	float FEnvelopeFollower::ProcessAudio(const float InAudioSample)
 	{
-		float Sample = (EnvMode != EPeakMode::Peak) ? InAudioSample * InAudioSample : FMath::Abs(InAudioSample);
-		float TimeSamples = (Sample > CurrentEnvelopeValue) ? AttackTimeSamples : ReleaseTimeSamples;
-		float NewEnvelopeValue = TimeSamples * (CurrentEnvelopeValue - Sample) + Sample;;
-		NewEnvelopeValue = Audio::UnderflowClamp(NewEnvelopeValue);
-		NewEnvelopeValue = FMath::Clamp(NewEnvelopeValue, 0.0f, 1.0f);
+		ProcessAudio(&InAudioSample, 1);
 
 		// Update and return the envelope value
-		return CurrentEnvelopeValue = NewEnvelopeValue;
+		return CurrentEnvelopeValue;
 	}
 
 	float FEnvelopeFollower::ProcessAudio(const float* InAudioBuffer, int32 InNumSamples)
 	{
-		for (int32 SampleIndex = 0; SampleIndex < InNumSamples; ++SampleIndex)
+		// MS/RMS
+		if (EnvMode == EPeakMode::MeanSquared || EnvMode == EPeakMode::RootMeanSquared)
 		{
-			ProcessAudioNonClamped(InAudioBuffer[SampleIndex]);
+			TAutoSlidingWindow<float> SlidingWindow(SumBuffer, TArrayView<const float>(InAudioBuffer, InNumSamples), ScratchBuffer, false);
+
+			for (auto& Window : SlidingWindow)
+			{
+				float CurrentMean;
+				
+				ArrayMeanSquared(Window, CurrentMean);
+
+				if (EnvMode == EPeakMode::RootMeanSquared)
+				{
+					CurrentMean = FMath::Sqrt(CurrentMean);
+				}
+
+				for (int j = 0; j < MeanWindowSize; ++j)
+				{
+					ProcessAudioNonClamped(CurrentMean);
+				}
+			}
 		}
-		return FMath::Clamp(CurrentEnvelopeValue, 0.0f, 1.0f);
+		// Peak mode
+		else
+		{
+			for (int32 SampleIndex = 0; SampleIndex < InNumSamples; ++SampleIndex)
+			{
+				ProcessAudioNonClamped(FMath::Abs(InAudioBuffer[SampleIndex]));
+			}
+		}
+		return CurrentEnvelopeValue = FMath::Clamp(CurrentEnvelopeValue, 0.0f, 1.0f);
 	}
 
 	float FEnvelopeFollower::ProcessAudio(const float* InAudioBuffer, float* OutAudioBuffer, int32 InNumSamples)
 	{
-		for (int32 SampleIndex = 0; SampleIndex < InNumSamples; ++SampleIndex)
+		// MS/RMS
+		if (EnvMode == EPeakMode::MeanSquared || EnvMode == EPeakMode::RootMeanSquared)
+		{
+			TAutoSlidingWindow<float> SlidingWindow(SumBuffer, TArrayView<const float>(InAudioBuffer, InNumSamples), ScratchBuffer, false);
+
+			for (auto& Window : SlidingWindow)
+			{
+				float CurrentMean;
+				ArrayMeanSquared(Window, CurrentMean);
+
+				if (EnvMode == EPeakMode::RootMeanSquared)
+				{
+					CurrentMean = FMath::Sqrt(CurrentMean);
+				}
+
+				for (int j = 0; j < MeanWindowSize; ++j)
+				{
+					ProcessAudioNonClamped(CurrentMean);
+				}
+			}
+		}
+		// Peak
+		else
 		{
-			OutAudioBuffer[SampleIndex] = ProcessAudioNonClamped(InAudioBuffer[SampleIndex]);
+			for (int32 SampleIndex = 0; SampleIndex < InNumSamples; ++SampleIndex)
+			{
+				OutAudioBuffer[SampleIndex] = ProcessAudioNonClamped(FMath::Abs(InAudioBuffer[SampleIndex]));
+			}
 		}
 
 		Audio::BufferRangeClampFast(OutAudioBuffer, InNumSamples, 0.0f, 1.0f);
@@ -109,9 +173,8 @@ namespace Audio
 
 	float FEnvelopeFollower::ProcessAudioNonClamped(const float InAudioSample)
 	{
-		float Sample = (EnvMode != EPeakMode::Peak) ? InAudioSample * InAudioSample : FMath::Abs(InAudioSample);
-		float TimeSamples = (Sample > CurrentEnvelopeValue) ? AttackTimeSamples : ReleaseTimeSamples;
-		float NewEnvelopeValue = TimeSamples * (CurrentEnvelopeValue - Sample) + Sample;;
+		float TimeSamples = (InAudioSample > CurrentEnvelopeValue) ? AttackTimeSamples : ReleaseTimeSamples;
+		float NewEnvelopeValue = TimeSamples * (CurrentEnvelopeValue - InAudioSample) + InAudioSample;
 		NewEnvelopeValue = Audio::UnderflowClamp(NewEnvelopeValue);
 
 		// Update and return the envelope value
diff --git a/Engine/Source/Runtime/SignalProcessing/Private/FloatArrayMath.cpp b/Engine/Source/Runtime/SignalProcessing/Private/FloatArrayMath.cpp
--- a/Engine/Source/Runtime/SignalProcessing/Private/FloatArrayMath.cpp
+++ b/Engine/Source/Runtime/SignalProcessing/Private/FloatArrayMath.cpp
@@ -108,6 +108,27 @@ namespace Audio
 		OutMean /= static_cast<float>(Num);
 	}
 
+	void ArrayMeanSquared(TArrayView<const float> InView, float& OutMean)
+	{
+		OutMean = 0.0f;
+
+		const int32 Num = InView.Num();
+
+		if (Num < 1)
+		{
+			return;
+		}
+
+		const float* DataPtr = InView.GetData();
+
+		for (int32 i = 0; i < Num; i++)
+		{
+			OutMean += DataPtr[i] * DataPtr[i];
+		}
+
+		OutMean /= static_cast<float>(Num);
+	}
+
 	void ArrayMeanFilter(TArrayView<const float> InView, int32 WindowSize, int32 WindowOrigin, TArray<float>& OutData)
 	{
 		// a quick but sinful implementation of a mean filter. encourages floating point rounding errors. 
diff --git a/Engine/Source/Runtime/SignalProcessing/Public/DSP/EnvelopeFollower.h b/Engine/Source/Runtime/SignalProcessing/Public/DSP/EnvelopeFollower.h
--- a/Engine/Source/Runtime/SignalProcessing/Public/DSP/EnvelopeFollower.h
+++ b/Engine/Source/Runtime/SignalProcessing/Public/DSP/EnvelopeFollower.h
@@ -3,6 +3,7 @@
 #pragma once
 
 #include "CoreMinimal.h"
+#include "SlidingWindow.h"
 
 namespace Audio
 {
@@ -24,13 +25,29 @@ namespace Audio
 	{
 	public:
 
+		static constexpr int32 DefaultWindowSize = 100;
+		static constexpr int32 DefaultHopSize = DefaultWindowSize / 2;
+		static_assert(DefaultHopSize > 0);
+
 		FEnvelopeFollower();
-		FEnvelopeFollower(const float InSampleRate, const float InAttackTimeMsec = 10.0f, const float InReleaseTimeMSec = 100.0f, const EPeakMode::Type InMode = EPeakMode::Peak, const bool bInIsAnalg = true);
+		FEnvelopeFollower(const float InSampleRate,
+			const float InAttackTimeMsec = 10.0f,
+			const float InReleaseTimeMSec = 100.0f,
+			const EPeakMode::Type InMode = EPeakMode::RootMeanSquared,
+			const int32 InWindowSizeForMean = DefaultWindowSize,
+			const int32 InHopSizeForMean = DefaultHopSize,
+			const bool bInIsAnalg = true);
 
 		virtual ~FEnvelopeFollower();
 
 		// Initialize the envelope follower
-		void Init(const float InSampleRate, const float InAttackTimeMsec = 10.0f, const float InReleaseTimeMSec = 100.0f, const EPeakMode::Type InMode = EPeakMode::Peak, const bool bInIsAnalg = true);
+		void Init(const float InSampleRate,
+			const float InAttackTimeMsec = 10.0f,
+			const float InReleaseTimeMSec = 100.0f,
+			const EPeakMode::Type InMode = EPeakMode::RootMeanSquared,
+			const int32 InWindowSizeForMean = DefaultWindowSize,
+			const int32 InHopSizeForMean = DefaultHopSize,
+			const bool bInIsAnalg = true);
 
 		// Resets the state of the envelope follower
 		void Reset();
@@ -67,6 +84,10 @@ namespace Audio
 
 	protected:
 		EPeakMode::Type EnvMode;
+		int32 MeanWindowSize;
+		int32 MeanHopSize;
+		TSlidingBuffer<float> SumBuffer;
+		TArray<float> ScratchBuffer;
 		float SampleRate;
 		float AttackTimeMsec;
 		float AttackTimeSamples;
diff --git a/Engine/Source/Runtime/SignalProcessing/Public/DSP/FloatArrayMath.h b/Engine/Source/Runtime/SignalProcessing/Public/DSP/FloatArrayMath.h
--- a/Engine/Source/Runtime/SignalProcessing/Public/DSP/FloatArrayMath.h
+++ b/Engine/Source/Runtime/SignalProcessing/Public/DSP/FloatArrayMath.h
@@ -24,6 +24,13 @@ namespace Audio
 	 */
 	SIGNALPROCESSING_API void ArrayMean(TArrayView<const float> InView, float& OutMean);
 
+	/** Mean Squared of array. Equivalent to Sum(InView * InView) / InView.Num()
+	 *
+	 *  InArray contains data to be analyzed.
+	 *  OutMean contains the result.
+	 */
+	SIGNALPROCESSING_API void ArrayMeanSquared(TArrayView<const float> InView, float& OutMean);
+
 	/** Mean filter of array.
 	 *
 	 *  Note: Uses standard biased mean estimator of Sum(x) / Count(x).
