Changed Parameter Order of Envelope Folower Constructor and Init Methods

diff --git a/Engine/Source/Runtime/SignalProcessing/Private/EnvelopeFollower.cpp b/Engine/Source/Runtime/SignalProcessing/Private/EnvelopeFollower.cpp
--- a/Engine/Source/Runtime/SignalProcessing/Private/EnvelopeFollower.cpp
+++ b/Engine/Source/Runtime/SignalProcessing/Private/EnvelopeFollower.cpp
@@ -30,11 +30,11 @@ namespace Audio
 		const float InAttackTimeMsec,
 		const float InReleaseTimeMSec,
 		const EPeakMode::Type InMode,
+		const bool bInIsAnalog,
 		const int32 InWindowSizeForMean,
-		const int32 InHopSizeForMean,
-		const bool bInIsAnalog) : SumBuffer(InWindowSizeForMean, InHopSizeForMean)
+		const int32 InHopSizeForMean) : SumBuffer(InWindowSizeForMean, InHopSizeForMean)
 	{
-		Init(InSampleRate, InAttackTimeMsec, InReleaseTimeMSec, InMode, InWindowSizeForMean, InHopSizeForMean, bInIsAnalog);
+		Init(InSampleRate, InAttackTimeMsec, InReleaseTimeMSec, InMode, bInIsAnalog,InWindowSizeForMean, InHopSizeForMean);
 	}
 
 	FEnvelopeFollower::~FEnvelopeFollower()
@@ -45,9 +45,9 @@ namespace Audio
 		const float InAttackTimeMsec, 
 		const float InReleaseTimeMSec, 
 		const EPeakMode::Type InMode, 
+		const bool bInIsAnalog,
 		const int32 InWindowSizeForMean,
-		const int32 InHopSizeForMean,
-		const bool bInIsAnalog)
+		const int32 InHopSizeForMean)
 	{
 		SampleRate = InSampleRate;
 
@@ -142,6 +142,7 @@ namespace Audio
 		{
 			TAutoSlidingWindow<float> SlidingWindow(SumBuffer, TArrayView<const float>(InAudioBuffer, InNumSamples), ScratchBuffer, false);
 
+			int32 SampleIndex = 0;
 			for (auto& Window : SlidingWindow)
 			{
 				float CurrentMean;
@@ -154,7 +155,8 @@ namespace Audio
 
 				for (int j = 0; j < MeanWindowSize; ++j)
 				{
-					ProcessAudioNonClamped(CurrentMean);
+					OutAudioBuffer[SampleIndex] = ProcessAudioNonClamped(CurrentMean);
+					++SampleIndex;
 				}
 			}
 		}
diff --git a/Engine/Source/Runtime/SignalProcessing/Public/DSP/EnvelopeFollower.h b/Engine/Source/Runtime/SignalProcessing/Public/DSP/EnvelopeFollower.h
--- a/Engine/Source/Runtime/SignalProcessing/Public/DSP/EnvelopeFollower.h
+++ b/Engine/Source/Runtime/SignalProcessing/Public/DSP/EnvelopeFollower.h
@@ -34,9 +34,9 @@ namespace Audio
 			const float InAttackTimeMsec = 10.0f,
 			const float InReleaseTimeMSec = 100.0f,
 			const EPeakMode::Type InMode = EPeakMode::RootMeanSquared,
+			const bool bInIsAnalg = true,
 			const int32 InWindowSizeForMean = DefaultWindowSize,
-			const int32 InHopSizeForMean = DefaultHopSize,
-			const bool bInIsAnalg = true);
+			const int32 InHopSizeForMean = DefaultHopSize);
 
 		virtual ~FEnvelopeFollower();
 
@@ -45,9 +45,9 @@ namespace Audio
 			const float InAttackTimeMsec = 10.0f,
 			const float InReleaseTimeMSec = 100.0f,
 			const EPeakMode::Type InMode = EPeakMode::RootMeanSquared,
+			const bool bInIsAnalg = true,
 			const int32 InWindowSizeForMean = DefaultWindowSize,
-			const int32 InHopSizeForMean = DefaultHopSize,
-			const bool bInIsAnalg = true);
+			const int32 InHopSizeForMean = DefaultHopSize);
 
 		// Resets the state of the envelope follower
 		void Reset();
