Fix argument order on dynamics preprocessor.

diff --git a/Engine/Plugins/Runtime/Metasound/Source/MetasoundStandardNodes/Private/MetasoundCompressorNode.cpp b/Engine/Plugins/Runtime/Metasound/Source/MetasoundStandardNodes/Private/MetasoundCompressorNode.cpp
--- a/Engine/Plugins/Runtime/Metasound/Source/MetasoundStandardNodes/Private/MetasoundCompressorNode.cpp
+++ b/Engine/Plugins/Runtime/Metasound/Source/MetasoundStandardNodes/Private/MetasoundCompressorNode.cpp
@@ -272,11 +272,11 @@ namespace Metasound
 
 			if (bUseSidechain)
 			{
-				Compressor.ProcessAudio(AudioInput->GetData(), AudioInput->Num(), AudioOutput->GetData(), EnvelopeOutput->GetData(), SidechainInput->GetData());
+				Compressor.ProcessAudio(AudioInput->GetData(), AudioInput->Num(), AudioOutput->GetData(), SidechainInput->GetData(), EnvelopeOutput->GetData());
 			}
 			else
 			{
-				Compressor.ProcessAudio(AudioInput->GetData(), AudioInput->Num(), AudioOutput->GetData(), EnvelopeOutput->GetData());
+				Compressor.ProcessAudio(AudioInput->GetData(), AudioInput->Num(), AudioOutput->GetData(), nullptr, EnvelopeOutput->GetData());
 			}
 
 			// Calculate Wet/Dry mix
diff --git a/Engine/Source/Runtime/SignalProcessing/Private/DynamicsProcessor.cpp b/Engine/Source/Runtime/SignalProcessing/Private/DynamicsProcessor.cpp
--- a/Engine/Source/Runtime/SignalProcessing/Private/DynamicsProcessor.cpp
+++ b/Engine/Source/Runtime/SignalProcessing/Private/DynamicsProcessor.cpp
@@ -284,6 +284,9 @@ namespace Audio
 
 	void FDynamicsProcessor::ProcessAudio(const float* InBuffer, const int32 InNumSamples, float* OutBuffer, const float* InKeyBuffer)
 	{
+		check(nullptr != InBuffer);
+		check(nullptr != OutBuffer);
+
 		const int32 NumChannels = GetNumChannels();
 		const int32 KeyNumChannels = GetKeyNumChannels();
 
@@ -307,12 +310,15 @@ namespace Audio
 		}
 	}
 
-	void FDynamicsProcessor::ProcessAudio(const float* InBuffer, const int32 InNumSamples, float* OutBuffer, float* OutEnvelope, const float* InKeyBuffer)
+	void FDynamicsProcessor::ProcessAudio(const float* InBuffer, const int32 InNumSamples, float* OutBuffer, const float* InKeyBuffer, float* OutEnvelope)
 	{
+		check(nullptr != InBuffer);
+		check(nullptr != OutBuffer);
+		check(nullptr != OutEnvelope);
+
 		const int32 NumChannels = GetNumChannels();
 		const int32 KeyNumChannels = GetKeyNumChannels();
 
-		check(OutEnvelope != nullptr);
 
 		if (InKeyBuffer)
 		{
diff --git a/Engine/Source/Runtime/SignalProcessing/Public/DSP/DynamicsProcessor.h b/Engine/Source/Runtime/SignalProcessing/Public/DSP/DynamicsProcessor.h
--- a/Engine/Source/Runtime/SignalProcessing/Public/DSP/DynamicsProcessor.h
+++ b/Engine/Source/Runtime/SignalProcessing/Public/DSP/DynamicsProcessor.h
@@ -69,7 +69,7 @@ namespace Audio
 		void ProcessAudioFrame(const float* InFrame, float* OutFrame, const float* InKeyFrame);
 		void ProcessAudioFrame(const float* InFrame, float* OutFrame, const float* InKeyFrame, float* OutGain);
 		void ProcessAudio(const float* InBuffer, const int32 InNumSamples, float* OutBuffer, const float* InKeyBuffer = nullptr);
-		void ProcessAudio(const float* InBuffer, const int32 InNumSamples, float* OutBuffer, float* OutEnvelope, const float* InKeyBuffer = nullptr);
+		void ProcessAudio(const float* InBuffer, const int32 InNumSamples, float* OutBuffer, const float* InKeyBuffer, float* OutEnvelope);
 		 
 
 	protected:
