Metasounds - Additive Synth - phase change interpolation and bugfix
- interpolate phase change over a block so that sinusoids don't need to be reset when phase is changed
- unit tests for sine wave table oscillator

diff --git a/Engine/Source/Runtime/SignalProcessing/Private/SineWaveTableOsc.cpp b/Engine/Source/Runtime/SignalProcessing/Private/SineWaveTableOsc.cpp
--- a/Engine/Source/Runtime/SignalProcessing/Private/SineWaveTableOsc.cpp
+++ b/Engine/Source/Runtime/SignalProcessing/Private/SineWaveTableOsc.cpp
@@ -17,9 +17,10 @@ namespace Audio
 
 	void FSineWaveTableOsc::Init(const float InSampleRate, const float InFrequencyHz, const float InPhase)
 	{
-		SampleRate = InSampleRate;
-		FrequencyHz = InFrequencyHz;
+		SampleRate = FMath::Clamp(InSampleRate, 0.0f, InSampleRate);
+		FrequencyHz = FMath::Clamp(InFrequencyHz, 0.0f, InFrequencyHz);
 		InitialPhase = FMath::Clamp(InPhase, 0.0f, 1.0f);
+		InstantaneousPhase = InitialPhase;
 
 		Reset();
 		UpdatePhaseIncrement();
@@ -27,25 +28,44 @@ namespace Audio
 
 	void FSineWaveTableOsc::SetSampleRate(const float InSampleRate)
 	{
-		SampleRate = InSampleRate;
+		SampleRate = FMath::Clamp(InSampleRate, 0.0f, InSampleRate);;
 		UpdatePhaseIncrement();
 	}
 
 	void FSineWaveTableOsc::Reset()
 	{
-		ReadIndex = InitialPhase * (WaveTableBuffer.Num() - 1.0f);		
+		ReadIndex = InitialPhase * WaveTableBuffer.Num();
+		while (ReadIndex >= WaveTableBuffer.Num())
+		{
+			ReadIndex -= WaveTableBuffer.Num();
+		}
 	}
 
 	void FSineWaveTableOsc::SetFrequencyHz(const float InFrequencyHz)
 	{
-		FrequencyHz = InFrequencyHz;
+		FrequencyHz = FMath::Clamp(InFrequencyHz, 0.0f, InFrequencyHz);;
 		UpdatePhaseIncrement();
 	}
 
 	void FSineWaveTableOsc::SetPhase(const float InPhase)
 	{
-		InitialPhase = FMath::Clamp(InPhase, 0.0f, 1.0f);
-		Reset();
+		float ClampedInPhase = FMath::Clamp(InPhase, 0.0f, 1.0f);
+		if (!FMath::IsNearlyEqual(ClampedInPhase, InitialPhase))
+		{
+			float PhaseDiff = ClampedInPhase - InitialPhase;
+			// InitialPhase will be set to ClampedInPhase once we interpolate the phase in Generate
+			InstantaneousPhase = ClampedInPhase;
+			// Increment ReadIndex by phase difference and wrap around if necessary
+			ReadIndex += PhaseDiff * WaveTableBuffer.Num();
+			while (ReadIndex >= WaveTableBuffer.Num())
+			{
+				ReadIndex -= WaveTableBuffer.Num();
+			}
+			while (ReadIndex < 0.0f)
+			{
+				ReadIndex += WaveTableBuffer.Num();
+			}
+		}
 	}
 
 	void FSineWaveTableOsc::UpdatePhaseIncrement()
@@ -53,27 +73,45 @@ namespace Audio
 		PhaseIncrement = (float)WaveTableBuffer.Num() * FrequencyHz / (float)SampleRate;
 	}
 
-	void FSineWaveTableOsc::Generate(float* OutSample)
+	void FSineWaveTableOsc::Generate(float* OutBuffer, const int32 NumSamples)
 	{
-		const int32 ReadIndexPrev = (int32)ReadIndex;
-		const float Alpha = ReadIndex - (float) ReadIndexPrev;
-
-		const int32 ReadIndexNext = (ReadIndexPrev + 1) % WaveTableBuffer.Num();
-		*OutSample = FMath::Lerp(WaveTableBuffer[ReadIndexPrev], WaveTableBuffer[ReadIndexNext], Alpha);
-		
-		// Increment ReadIndex and wrap around if necessary
-		ReadIndex += PhaseIncrement;
-		while (ReadIndex >= WaveTableBuffer.Num())
+		// Interpolate phase shift over the block 
+		const float HalfNumSamples = FMath::Max(1.f, static_cast<float>(NumSamples / 2));
+		float PhaseShiftTrianglePeak = 0.0f;
+		bool InterpolatePhase = !FMath::IsNearlyEqual(InitialPhase, InstantaneousPhase);
+		if (InterpolatePhase)
 		{
-			ReadIndex -= WaveTableBuffer.Num();
+			float PhaseDiff = InstantaneousPhase - InitialPhase;
+			PhaseShiftTrianglePeak = PhaseDiff * 2 / NumSamples;
+			InitialPhase = InstantaneousPhase;
 		}
-	}
 
-	void FSineWaveTableOsc::Generate(float* OutBuffer, const int32 NumSamples)
-	{
 		for (int32 SampleIndex = 0; SampleIndex < NumSamples; ++SampleIndex)
 		{
-			Generate(&OutBuffer[SampleIndex]);
+			// Interpolate between two samples
+			const int32 ReadIndexPrev = (int32)ReadIndex;
+			const float Alpha = ReadIndex - (float)ReadIndexPrev;
+
+			const int32 ReadIndexNext = (ReadIndexPrev + 1) % WaveTableBuffer.Num();
+			OutBuffer[SampleIndex] = FMath::Lerp(WaveTableBuffer[ReadIndexPrev], WaveTableBuffer[ReadIndexNext], Alpha);
+			
+			// Increment ReadIndex and wrap around if necessary
+			if (InterpolatePhase)
+			{
+				// Interpolate phase over the block using a triangle 
+				float BlockFraction = static_cast<float>(SampleIndex);
+				float PhaseShift = PhaseShiftTrianglePeak * FMath::Abs((FMath::Abs(1.f - BlockFraction / HalfNumSamples) - 1.f));
+				ReadIndex += PhaseIncrement + PhaseShift;
+			}
+			else
+			{
+				ReadIndex += PhaseIncrement;
+			}
+
+			while (ReadIndex >= WaveTableBuffer.Num())
+			{
+				ReadIndex -= WaveTableBuffer.Num();
+			}
 		}
 	}
 
diff --git a/Engine/Source/Runtime/SignalProcessing/Public/DSP/SineWaveTableOsc.h b/Engine/Source/Runtime/SignalProcessing/Public/DSP/SineWaveTableOsc.h
--- a/Engine/Source/Runtime/SignalProcessing/Public/DSP/SineWaveTableOsc.h
+++ b/Engine/Source/Runtime/SignalProcessing/Public/DSP/SineWaveTableOsc.h
@@ -37,8 +37,6 @@ namespace Audio
 		// Sets the phase of the wave table oscillator.
 		void SetPhase(const float InPhase);
 
-		// Processes the wave table and outputs a sample 
-		void Generate(float* OutSample);
 		// Processes the wave table and fills a buffer
 		void Generate(float* OutBuffer, const int32 NumSamples);
 
@@ -63,7 +61,10 @@ namespace Audio
 		// The phase increment (based on frequency)
 		float PhaseIncrement = 0.0f;
 
-		// The initial phase, will be clamped to range 0.0 - 1.0f
+		// The initial/cached phase, will be clamped to range 0.0 - 1.0f
 		float InitialPhase = 0.0f;
+
+		// The instantaneous phase, clamped to range 0.0 - 1.0f
+		float InstantaneousPhase = 0.0f;
 	};
 }
