Fix audio envelope DSP implementation and performance. Fixes multiple perf related issues related to audio envelopes.

diff --git a/Engine/Plugins/Runtime/AudioModulation/Source/AudioModulation/Private/Generators/SoundModulationEnvelopeFollower.cpp b/Engine/Plugins/Runtime/AudioModulation/Source/AudioModulation/Private/Generators/SoundModulationEnvelopeFollower.cpp
--- a/Engine/Plugins/Runtime/AudioModulation/Source/AudioModulation/Private/Generators/SoundModulationEnvelopeFollower.cpp
+++ b/Engine/Plugins/Runtime/AudioModulation/Source/AudioModulation/Private/Generators/SoundModulationEnvelopeFollower.cpp
@@ -2,6 +2,7 @@
 
 #include "Generators/SoundModulationEnvelopeFollower.h"
 
+#include "Algo/MaxElement.h"
 #include "AudioDeviceManager.h"
 #include "AudioMixerDevice.h"
 #include "AudioModulation.h"
@@ -26,7 +27,7 @@ namespace AudioModulation
 
 		if (Params.AudioBus)
 		{
-			AudioRenderThreadCommand([this, DeviceId = InDeviceId, BusId = Params.AudioBus->GetUniqueID()]()
+			AudioRenderThreadCommand([this, DeviceId = InDeviceId, NumChannels = Params.AudioBus->GetNumChannels(), BusId = Params.AudioBus->GetUniqueID()]()
 			{
 				if (FAudioDeviceManager* DeviceManager = FAudioDeviceManager::Get())
 				{
@@ -34,14 +35,23 @@ namespace AudioModulation
 					if (FMixerDevice* MixerDevice = static_cast<FMixerDevice*>(AudioDevice))
 					{
 						AudioBusPatch = MixerDevice->AddPatchForAudioBus(BusId, Params.Gain);
-						EnvelopeFollower = FEnvelopeFollower(MixerDevice->SampleRate, Params.AttackTime, Params.ReleaseTime, Audio::EPeakMode::Peak);
+
+						FEnvelopeFollowerInitParams EnvelopeFollowerInitParams;
+						EnvelopeFollowerInitParams.SampleRate = MixerDevice->SampleRate;
+						EnvelopeFollowerInitParams.NumChannels = NumChannels;
+						EnvelopeFollowerInitParams.AttackTimeMsec = Params.AttackTime;
+						EnvelopeFollowerInitParams.ReleaseTimeMsec = Params.ReleaseTime;
+						EnvelopeFollowerInitParams.Mode = EPeakMode::Peak;
+
+						EnvelopeFollower = FEnvelopeFollower(EnvelopeFollowerInitParams);
+
 						if (Params.bInvert)
 						{
-							CurrentValue = 1.0f - EnvelopeFollower.GetCurrentValue();
+							CurrentValue = 1.0f;
 						}
 						else
 						{
-							CurrentValue = EnvelopeFollower.GetCurrentValue();
+							CurrentValue = 0.f;
 						}
 					}
 				}
@@ -61,11 +71,11 @@ namespace AudioModulation
 
 	void FEnvelopeFollowerGenerator::Update(double InElapsed)
 	{
-		float Amp = 0.0f;
-
 		if (AudioBusPatch.IsValid())
 		{
-			int32 NumSamples = AudioBusPatch->GetNumSamplesAvailable();
+			const int32 NumSamples = AudioBusPatch->GetNumSamplesAvailable();
+			const int32 NumFrames = NumSamples / EnvelopeFollower.GetNumChannels();
+
 			if (NumSamples > 0)
 			{
 				TempBuffer.Reset();
@@ -73,40 +83,35 @@ namespace AudioModulation
 
 				AudioBusPatch->PopAudio(TempBuffer.GetData(), NumSamples, true /* bUseLatestAudio */);
 
-				// If NumSamples is below 4 samples, just take first sample
-				static const int32 SimdMultiple = 4;
-				if (NumSamples % SimdMultiple == 0)
-				{
-					Amp = Audio::BufferGetAverageAbsValue(TempBuffer.GetData(), NumSamples);
-				}
-				// If NumSamples is not a multiple of 4, do it the non-SIMD way
-				else
-				{
-					Amp = 0.0f;
-					for (int32 i = 0; i < NumSamples; ++i)
-					{
-						Amp += FMath::Abs(TempBuffer[i]);
-					}
-					Amp /= NumSamples;
-				}
+				EnvelopeFollower.ProcessAudio(TempBuffer.GetData(), NumFrames);
 			}
-		}
 
-		EnvelopeFollower.ProcessAudio(Amp);
-		if (Params.bInvert)
-		{
-			CurrentValue = 1.0f - EnvelopeFollower.GetCurrentValue();
+			float MaxValue = 0.f;
+			if (const float* MaxEnvelopePtr = Algo::MaxElement(EnvelopeFollower.GetEnvelopeValues()))
+			{
+				MaxValue = FMath::Clamp(*MaxEnvelopePtr, 0.f, 1.f);
+			}
+
+			if (Params.bInvert)
+			{
+				CurrentValue = 1.0f - MaxValue;
+			}
+			else
+			{
+				CurrentValue = MaxValue;
+			}
 		}
 		else
 		{
-			CurrentValue = EnvelopeFollower.GetCurrentValue();
+			CurrentValue = 0.f;
 		}
+
 	}
 
 #if !UE_BUILD_SHIPPING
 	void FEnvelopeFollowerGenerator::GetDebugValues(TArray<FString>& OutDebugValues) const
 	{
-		OutDebugValues.Add(FString::Printf(TEXT("%.4f"), EnvelopeFollower.GetCurrentValue()));
+		OutDebugValues.Add(FString::Printf(TEXT("%.4f"), GetValue()));
 		OutDebugValues.Add(FString::Printf(TEXT("%.4f"), Params.Gain));
 		OutDebugValues.Add(FString::Printf(TEXT("%.4f"), Params.AttackTime));
 		OutDebugValues.Add(FString::Printf(TEXT("%.4f"), Params.ReleaseTime));
diff --git a/Engine/Plugins/Runtime/AudioSynesthesia/Source/AudioSynesthesiaCore/Private/MeterAnalyzer.cpp b/Engine/Plugins/Runtime/AudioSynesthesia/Source/AudioSynesthesiaCore/Private/MeterAnalyzer.cpp
--- a/Engine/Plugins/Runtime/AudioSynesthesia/Source/AudioSynesthesiaCore/Private/MeterAnalyzer.cpp
+++ b/Engine/Plugins/Runtime/AudioSynesthesia/Source/AudioSynesthesiaCore/Private/MeterAnalyzer.cpp
@@ -10,75 +10,120 @@
 namespace Audio
 {
 	FMeterAnalyzer::FMeterAnalyzer(float InSampleRate, int32 InNumChannels, const FMeterAnalyzerSettings& InSettings)
+	: EnvelopeFollower(FEnvelopeFollowerInitParams{InSampleRate, InNumChannels, static_cast<float>(InSettings.MeterAttackTime), static_cast<float>(InSettings.MeterReleaseTime), InSettings.MeterPeakMode, true /* bInIsAnalog */})
+	, Settings(InSettings)
+	, SampleRate(InSampleRate)
+	, NumChannels(InNumChannels)
 	{
-		NumChannels = InNumChannels;
-		SampleRate = InSampleRate;
-		Settings = InSettings;
-		check(NumChannels > 0);
-		check(SampleRate > 0.0f);
-
-		PeakEnvDataPerChannel.AddDefaulted(NumChannels);
-		PeakDataPerChannel.AddDefaulted(NumChannels);
-		ClippingDataPerChannel.AddDefaulted(NumChannels);
-
-		for (int32 ChannelIndex = 0; ChannelIndex < NumChannels; ++ChannelIndex)
+		if (!ensure(NumChannels > 0))
+		{
+			NumChannels = 1;
+		}
+		if (!ensure(SampleRate > 0.f))
 		{
-			MeterEnvelopeFollowers.Add(MakeUnique<FEnvelopeFollower>(InSampleRate, (float)InSettings.MeterAttackTime, (float)InSettings.MeterReleaseTime, InSettings.MeterPeakMode, true));
+			SampleRate = 48000.f;
 		}
+
+		EnvelopeDataPerChannel.AddDefaulted(NumChannels);
+		ClippingDataPerChannel.AddDefaulted(NumChannels);
+		PeakHoldFrames = FMath::Max(1, FMath::RoundToInt(static_cast<float>(InSettings.PeakHoldTime) * SampleRate / 1000.f));
 	}
 
 	FMeterAnalyzerResults FMeterAnalyzer::ProcessAudio(TArrayView<const float> InSampleView)
 	{
 		// Feed audio through the envelope followers
-		for (int32 SampleIndex = 0; SampleIndex < InSampleView.Num(); SampleIndex += NumChannels)
+		const float* InData = InSampleView.GetData();
+		const int32 NumSamples = InSampleView.Num();
+		const int32 NumFrames = NumSamples / NumChannels;
+
+		check((NumSamples % NumChannels) == 0);
+
+		// Reset the clipping data for next block of audio
+		for (int32 ChannelIndex = 0; ChannelIndex < NumChannels; ChannelIndex++)
 		{
-			for (int32 ChannelIndex = 0; ChannelIndex < NumChannels; ++ChannelIndex)
-			{
-				float Sample = InSampleView[SampleIndex + ChannelIndex];
+			ClippingDataPerChannel[ChannelIndex].ClippingValue = 0.0f;
+			ClippingDataPerChannel[ChannelIndex].NumSamplesClipping = 0;
+		}
 
-				// Check for clipping in this block of audio
+		for (int32 i = 0; i < NumSamples; i++)
+		{
+			if (FMath::Abs(InData[i]) > Settings.ClippingThreshold)
+			{
+				const int32 ChannelIndex = i % NumChannels;
 				FClippingData& ClippingData = ClippingDataPerChannel[ChannelIndex];
-				if (Sample > Settings.ClippingThreshold)
-				{
-					// Track how many samples clipped in this block
-					ClippingData.NumSamplesClipping++;
 
-					// Track the max clipping value in this block of audio
-					ClippingData.ClippingValue = FMath::Max(ClippingData.ClippingValue, Sample);
-				}
+				// Track how many samples clipped in this block
+				ClippingData.NumSamplesClipping++;
+
+				// Track the max clipping value in this block of audio
+				ClippingData.ClippingValue = FMath::Max(ClippingData.ClippingValue, InData[i]);
+			}
+		}
+
+		EnvelopeBuffer.Reset(NumSamples);
+		EnvelopeBuffer.AddUninitialized(NumSamples);
+		EnvelopeFollower.ProcessAudio(InData, NumFrames, EnvelopeBuffer.GetData());
+
+		// Find peaks in the envelope.
+		const float* Envelope = EnvelopeBuffer.GetData();
+		for (int32 ChannelIndex = 0; ChannelIndex < NumChannels; ChannelIndex++)
+		{
+			FEnvelopeData& EnvelopeData = EnvelopeDataPerChannel[ChannelIndex];
+			// Reset max envelope data. 
+			EnvelopeData.MaxEnvelopeValue = 0.f;
+			if (EnvelopeData.FramesUntilPeakReset < 1)
+			{
+				EnvelopeData.PeakValue = 0;
+			}
 
-				TUniquePtr<FEnvelopeFollower>& MeterEnvelopeFollower = MeterEnvelopeFollowers[ChannelIndex];
-				float NewPeakEnvValue = MeterEnvelopeFollower->ProcessAudio(Sample);
+			// Temp variables to hold onto prior frame's envelope info.
+			float PriorValue = EnvelopeData.PriorEnvelopeValue;
+			bool bPriorSlopeIsPositive = EnvelopeData.bPriorEnvelopeSlopeIsPositive;
 
-				FPeakEnvelopeData& PeakEnvData = PeakEnvDataPerChannel[ChannelIndex];
-				FPeakData& PeakData = PeakDataPerChannel[ChannelIndex];
+			for (int32 SampleIndex = ChannelIndex; SampleIndex < NumSamples; SampleIndex += NumChannels)
+			{
+				// Get maximum value in envelope.
+				EnvelopeData.MaxEnvelopeValue = FMath::Max(Envelope[SampleIndex], EnvelopeData.MaxEnvelopeValue);
 
 				// Positive if we're rising
-				bool bNewSlopeIsPositive = (NewPeakEnvValue - PeakEnvData.Value) > 0.0f;
+				bool bNewSlopeIsNegative = (Envelope[SampleIndex] - PriorValue) >= 0.0f;
+				bool bIsPeak = bNewSlopeIsNegative && bPriorSlopeIsPositive;
 
 				// Detect if the new peak envelope value is less than the previous peak envelope value (i.e. local maximum)
-				if (!bNewSlopeIsPositive && PeakEnvData.bEnvelopeSlopeIsPositive)
+				if (bIsPeak)
 				{
-					// Get the current time
-					float CurrentTimeSec = (PeakEnvData.SampleCount + SampleIndex) / SampleRate;
+					const int32 FrameIndex = SampleIndex / NumChannels;
 
 					// Check if the peak value is larger than the current peak value or if enough time has elapsed to store a new peak
-					if (PeakData.Value < PeakEnvData.Value || (CurrentTimeSec - PeakData.StartTime) >= (0.001f * Settings.PeakHoldTime))
+					if ((EnvelopeData.PeakValue < Envelope[SampleIndex]) || (FrameIndex > EnvelopeData.FramesUntilPeakReset))
 					{
-						PeakData.Value = PeakEnvData.Value;
-						PeakData.StartTime = CurrentTimeSec;
+						EnvelopeData.PeakValue = Envelope[SampleIndex];
+						EnvelopeData.FramesUntilPeakReset = FrameIndex + PeakHoldFrames;
 					}
 				}
 
 				// Update the current peak envelope data to the new value
-				PeakEnvData.Value = NewPeakEnvValue;
-				PeakEnvData.bEnvelopeSlopeIsPositive = bNewSlopeIsPositive;
+				PriorValue = Envelope[SampleIndex];
+				bPriorSlopeIsPositive = !bNewSlopeIsNegative;
+			}
+
+			// Save for next call
+			EnvelopeData.PriorEnvelopeValue = PriorValue;
+			EnvelopeData.bPriorEnvelopeSlopeIsPositive = bPriorSlopeIsPositive;
+
+			if (EnvelopeData.FramesUntilPeakReset > 0)
+			{
+				// Decrement frame counters for next call.
+				EnvelopeData.FramesUntilPeakReset -= NumFrames;
 			}
 		}
 
 		// Build the results data
 		FMeterAnalyzerResults Results;
 
+		FrameCounter += NumFrames;
+		Results.TimeSec = static_cast<float>(FrameCounter) / SampleRate;
+
 		for (int32 ChannelIndex = 0; ChannelIndex < NumChannels; ++ChannelIndex)
 		{
 			// Add clipping data
@@ -87,21 +132,11 @@ namespace Audio
 			Results.ClippingValues.Add(ClippingData.ClippingValue);
 			Results.NumSamplesClipping.Add(ClippingData.NumSamplesClipping);
 
-			// Reset the clipping data for next block of audio
-			ClippingData.ClippingValue = 0.0f;
-			ClippingData.NumSamplesClipping = 0;
-
 			// update sample count of the env data
-			FPeakEnvelopeData& PeakEnvData = PeakEnvDataPerChannel[ChannelIndex];
-			PeakEnvData.SampleCount += InSampleView.Num();
-
-			TUniquePtr<FEnvelopeFollower>& MeterEnvelopeFollower = MeterEnvelopeFollowers[ChannelIndex];
-			Results.MeterValues.Add(MeterEnvelopeFollower->GetCurrentValue());
-
-			FPeakData& PeakData = PeakDataPerChannel[ChannelIndex];
-			Results.PeakValues.Add(PeakData.Value);
+			const FEnvelopeData& EnvData = EnvelopeDataPerChannel[ChannelIndex];
 
-			Results.TimeSec = PeakEnvData.SampleCount / SampleRate;
+			Results.MeterValues.Add(EnvData.MaxEnvelopeValue);
+			Results.PeakValues.Add(EnvData.PeakValue);
 		}
 
 		return Results;
diff --git a/Engine/Plugins/Runtime/AudioSynesthesia/Source/AudioSynesthesiaCore/Public/MeterAnalyzer.h b/Engine/Plugins/Runtime/AudioSynesthesia/Source/AudioSynesthesiaCore/Public/MeterAnalyzer.h
--- a/Engine/Plugins/Runtime/AudioSynesthesia/Source/AudioSynesthesiaCore/Public/MeterAnalyzer.h
+++ b/Engine/Plugins/Runtime/AudioSynesthesia/Source/AudioSynesthesiaCore/Public/MeterAnalyzer.h
@@ -53,7 +53,8 @@ namespace Audio
 
 	protected:
 		// Envelope follower per channel
-		TArray<TUniquePtr<FEnvelopeFollower>> MeterEnvelopeFollowers;
+		FEnvelopeFollower EnvelopeFollower;
+		FAlignedFloatBuffer EnvelopeBuffer;
 
 		// Per-channel clipping data
 		struct FClippingData
@@ -65,31 +66,25 @@ namespace Audio
 		TArray<FClippingData> ClippingDataPerChannel;
 
 		// State to track the peak data
-		struct FPeakData
+		struct FEnvelopeData
 		{
-			float Value = 0.0f;
-			float StartTime = 0.0f;
-		};
-		// Per-channel peak data
-		TArray<FPeakData> PeakDataPerChannel;
-
-		struct FPeakEnvelopeData
-		{
-			// Current envelope data
-			float Value = 0.0f;
-
-			// Running sample count of the envelope data
-			int32 SampleCount = 0;
+			float MaxEnvelopeValue = 0.f;
+			float PeakValue = 0.0f;
+			int32 FramesUntilPeakReset = 0;
 
 			// Current slope of the envelope
-			bool bEnvelopeSlopeIsPositive = false;
+			float PriorEnvelopeValue = 0.f;
+			bool bPriorEnvelopeSlopeIsPositive = false;
 		};
-		// Per-channel envelope data
-		TArray<FPeakEnvelopeData> PeakEnvDataPerChannel;
+		// Per-channel peak data
+		TArray<FEnvelopeData> EnvelopeDataPerChannel;
+
 
 		FMeterAnalyzerSettings Settings;
 		float SampleRate = 0.0f;
 		int32 NumChannels = 0;
+		int32 PeakHoldFrames = 0;
+		int64 FrameCounter = 0;
 	};
 }
 
diff --git a/Engine/Plugins/Runtime/Metasound/Source/MetasoundStandardNodes/Private/MetasoundEnvelopeFollowerNode.cpp b/Engine/Plugins/Runtime/Metasound/Source/MetasoundStandardNodes/Private/MetasoundEnvelopeFollowerNode.cpp
--- a/Engine/Plugins/Runtime/Metasound/Source/MetasoundStandardNodes/Private/MetasoundEnvelopeFollowerNode.cpp
+++ b/Engine/Plugins/Runtime/Metasound/Source/MetasoundStandardNodes/Private/MetasoundEnvelopeFollowerNode.cpp
@@ -2,6 +2,7 @@
 
 #include "MetasoundEnvelopeFollowerNode.h"
 
+#include "Algo/MaxElement.h"
 #include "Internationalization/Text.h"
 #include "MetasoundAudioBuffer.h"
 #include "MetasoundExecutableOperator.h"
@@ -106,7 +107,14 @@ namespace Metasound
 		PrevAttackTime = FMath::Max(FTime::ToMilliseconds(*AttackTimeInput), 0.0);
 		PrevReleaseTime = FMath::Max(FTime::ToMilliseconds(*ReleaseTimeInput), 0.0);
 
-		EnvelopeFollower.Init(InSettings.GetSampleRate(), PrevAttackTime, PrevReleaseTime);
+		Audio::FEnvelopeFollowerInitParams EnvelopeParamsInitParams;
+	
+		EnvelopeParamsInitParams.SampleRate = InSettings.GetSampleRate();
+		EnvelopeParamsInitParams.NumChannels = 1;
+		EnvelopeParamsInitParams.AttackTimeMsec = PrevAttackTime;
+		EnvelopeParamsInitParams.ReleaseTimeMsec = PrevReleaseTime;
+
+		EnvelopeFollower.Init(EnvelopeParamsInitParams);
 	}
 
 	FDataReferenceCollection FEnvelopeFollowerOperator::GetInputs() const
@@ -164,8 +172,15 @@ namespace Metasound
 		}
 
 		// Process the audio through the envelope follower
-		EnvelopeFollower.ProcessAudio(AudioInput->GetData(), EnvelopeAudioOutput->GetData(), AudioInput->Num());
-		*EnvelopeFloatOutput = EnvelopeFollower.GetCurrentValue();
+		EnvelopeFollower.ProcessAudio(AudioInput->GetData(), AudioInput->Num(), EnvelopeAudioOutput->GetData());
+		if (const float* MaxElement = Algo::MaxElement(EnvelopeFollower.GetEnvelopeValues()))
+		{
+			*EnvelopeFloatOutput = *MaxElement;
+		}
+		else
+		{
+			*EnvelopeFloatOutput = 0.f;
+		}
 	}
 
 	const FVertexInterface& FEnvelopeFollowerOperator::GetVertexInterface()
diff --git a/Engine/Plugins/Runtime/Synthesis/Source/Synthesis/Classes/SourceEffects/SourceEffectEnvelopeFollower.h b/Engine/Plugins/Runtime/Synthesis/Source/Synthesis/Classes/SourceEffects/SourceEffectEnvelopeFollower.h
--- a/Engine/Plugins/Runtime/Synthesis/Source/Synthesis/Classes/SourceEffects/SourceEffectEnvelopeFollower.h
+++ b/Engine/Plugins/Runtime/Synthesis/Source/Synthesis/Classes/SourceEffects/SourceEffectEnvelopeFollower.h
@@ -63,7 +63,7 @@ public:
 
 protected:
 
-	Audio::FEnvelopeFollower EnvelopeFollower;
+	Audio::FInlineEnvelopeFollower EnvelopeFollower;
 	float CurrentEnvelopeValue;
 	uint32 OwningPresetUniqueId;
 	uint32 InstanceId;
diff --git a/Engine/Plugins/Runtime/Synthesis/Source/Synthesis/Private/SourceEffectEnvelopeFollower.cpp b/Engine/Plugins/Runtime/Synthesis/Source/Synthesis/Private/SourceEffectEnvelopeFollower.cpp
--- a/Engine/Plugins/Runtime/Synthesis/Source/Synthesis/Private/SourceEffectEnvelopeFollower.cpp
+++ b/Engine/Plugins/Runtime/Synthesis/Source/Synthesis/Private/SourceEffectEnvelopeFollower.cpp
@@ -220,9 +220,12 @@ void FSourceEffectEnvelopeFollower::Init(const FSoundEffectSourceInitData& InitD
 	FramesToNotify = 1024;
 	FrameCount = 0;
 	bIsActive = true;
-	EnvelopeFollower.Init(InitData.SampleRate);
 	CurrentEnvelopeValue = 0.0f;
 	NumChannels = InitData.NumSourceChannels;
+
+	Audio::FInlineEnvelopeFollowerInitParams EnvelopeFollowerInitParams;
+	EnvelopeFollowerInitParams.SampleRate = InitData.SampleRate;
+	EnvelopeFollower.Init(EnvelopeFollowerInitParams);
 }
 
 FSourceEffectEnvelopeFollower::~FSourceEffectEnvelopeFollower()
@@ -258,11 +261,11 @@ void FSourceEffectEnvelopeFollower::ProcessAudio(const FSoundEffectSourceInputDa
 			SampleValue *= 0.5f;
 		}
 
-		CurrentEnvelopeValue = EnvelopeFollower.ProcessAudio(SampleValue);
+		CurrentEnvelopeValue = EnvelopeFollower.ProcessSample(SampleValue);
 
 		if ((FrameCount++ & (FramesToNotify - 1)) == 0)
 		{
-			SourceEffectEnvFollowerNotifier->UpdateEnvFollowerInstance(OwningPresetUniqueId, InstanceId, CurrentEnvelopeValue);
+			SourceEffectEnvFollowerNotifier->UpdateEnvFollowerInstance(OwningPresetUniqueId, InstanceId, FMath::Clamp(CurrentEnvelopeValue, 0.f, 1.f));
 		}
 	}
 }
diff --git a/Engine/Plugins/Runtime/Synthesis/Source/Synthesis/Private/SourceEffectFilter.cpp b/Engine/Plugins/Runtime/Synthesis/Source/Synthesis/Private/SourceEffectFilter.cpp
--- a/Engine/Plugins/Runtime/Synthesis/Source/Synthesis/Private/SourceEffectFilter.cpp
+++ b/Engine/Plugins/Runtime/Synthesis/Source/Synthesis/Private/SourceEffectFilter.cpp
@@ -122,9 +122,13 @@ void FSourceEffectFilter::OnPresetChanged()
 					NewModData.MaxResModValue = BusModulationSettings.MaxResonanceModulation;
 					NewModData.EnvelopeGain = BusModulationSettings.EnvelopeGainMultiplier;
 
-					NewModData.AudioBusEnvelopeFollower.Init(SampleRate);
-					NewModData.AudioBusEnvelopeFollower.SetAttackTime(BusModulationSettings.EnvelopeFollowerAttackTimeMsec);
-					NewModData.AudioBusEnvelopeFollower.SetReleaseTime(BusModulationSettings.EnvelopeFollowerReleaseTimeMsec);
+					Audio::FEnvelopeFollowerInitParams EnvelopeFollowerInitParams;
+					EnvelopeFollowerInitParams.SampleRate = SampleRate;
+					EnvelopeFollowerInitParams.NumChannels = 1;
+					EnvelopeFollowerInitParams.AttackTimeMsec = BusModulationSettings.EnvelopeFollowerAttackTimeMsec;
+					EnvelopeFollowerInitParams.ReleaseTimeMsec = BusModulationSettings.EnvelopeFollowerReleaseTimeMsec;
+
+					NewModData.AudioBusEnvelopeFollower.Init(EnvelopeFollowerInitParams);
 				}
 			}
 		}
@@ -161,13 +165,28 @@ void FSourceEffectFilter::ProcessAudio(const FSoundEffectSourceInputData& InData
 			Audio::BufferSum2ChannelToMonoFast(ScratchModBuffer, ScratchEnvFollowerBuffer);
 			Audio::MultiplyBufferByConstantInPlace(ScratchEnvFollowerBuffer, 0.5f);
 
-			EnvelopeValue = Mod.AudioBusEnvelopeFollower.ProcessAudio(ScratchEnvFollowerBuffer.GetData(), InData.NumSamples);
+			Mod.AudioBusEnvelopeFollower.ProcessAudio(ScratchEnvFollowerBuffer.GetData(), InData.NumSamples);
+		}
+		else if (NumChannels == 1)
+		{
+			Mod.AudioBusEnvelopeFollower.ProcessAudio(ScratchModBuffer.GetData(), InData.NumSamples);
+		}
+		else
+		{
+			checkNoEntry();
+		}
+
+		const TArray<float>& CurrentEnvelopeValues = Mod.AudioBusEnvelopeFollower.GetEnvelopeValues();
+		if (ensure(CurrentEnvelopeValues.Num() == 1))
+		{
+			EnvelopeValue = FMath::Clamp(CurrentEnvelopeValues[0], 0.f, 1.f);
 		}
 		else
 		{
-			EnvelopeValue = Mod.AudioBusEnvelopeFollower.ProcessAudio(ScratchModBuffer.GetData(), InData.NumSamples);
+			EnvelopeValue = 0.f;
 		}
 
+
 		if (Mod.FilterParam == ESourceEffectFilterParam::FilterFrequency)
 		{
 			ModFrequency += FMath::Lerp(Mod.MinFreqModValue, Mod.MaxFreqModValue, FMath::Clamp(EnvelopeValue * Mod.EnvelopeGain, 0.0f, 1.0f));;
diff --git a/Engine/Plugins/Runtime/TimeSynth/Source/TimeSynth/Private/TimeSynthComponent.cpp b/Engine/Plugins/Runtime/TimeSynth/Source/TimeSynth/Private/TimeSynthComponent.cpp
--- a/Engine/Plugins/Runtime/TimeSynth/Source/TimeSynth/Private/TimeSynthComponent.cpp
+++ b/Engine/Plugins/Runtime/TimeSynth/Source/TimeSynth/Private/TimeSynthComponent.cpp
@@ -1,9 +1,11 @@
 // Copyright Epic Games, Inc. All Rights Reserved.
 
 #include "TimeSynthComponent.h"
+
+#include "Algo/MaxElement.h"
 #include "AudioThread.h"
-#include "TimeSynthModule.h"
 #include "Containers/UnrealString.h"
+#include "TimeSynthModule.h"
 
 static const float MinAudibleLog = -55;
 
@@ -422,7 +424,10 @@ bool UTimeSynthComponent::Init(int32& InSampleRate)
 	DynamicsProcessor.SetProcessingMode(Audio::EDynamicsProcessingMode::Compressor);
 
 	// Init and update the envelope follower settings
-	EnvelopeFollower.Init(InSampleRate);
+	Audio::FEnvelopeFollowerInitParams EnvelopeFollowerInitParams;
+	EnvelopeFollowerInitParams.SampleRate = InSampleRate;
+	EnvelopeFollowerInitParams.NumChannels = NumChannels;
+	EnvelopeFollower.Init(EnvelopeFollowerInitParams);
 	UpdateEnvelopeFollower();
 
 	// Set the default quantization settings
@@ -492,7 +497,7 @@ int32 UTimeSynthComponent::OnGenerateAudio(float* OutAudio, int32 NumSamples)
 	// Update the decoder
 	SoundWaveDecoder.UpdateRenderThread();
 
-	int32 NumFrames = NumSamples / NumChannels;
+	const int32 NumFrames = NumSamples / NumChannels;
 
 	// Perform event quantization notifications
 	// This will use the NumFramesPerCallback to evaluate what queued up events need 
@@ -612,10 +617,15 @@ int32 UTimeSynthComponent::OnGenerateAudio(float* OutAudio, int32 NumSamples)
 	// Feed audio through the envelope follower if it's enabled
 	if (bIsEnvelopeFollowerEnabled_AudioRenderThread)
 	{
-		for (int32 SampleIndex = 0; SampleIndex < NumSamples; SampleIndex += NumChannels)
+		EnvelopeFollower.ProcessAudio(OutAudio, NumFrames);
+
+		if (const float* MaxEnvValue = Algo::MaxElement(EnvelopeFollower.GetEnvelopeValues()))
+		{
+			CurrentEnvelopeValue = FMath::Clamp(*MaxEnvValue, 0.f, 1.f);
+		}
+		else
 		{
-			float InputSample = 0.5f * (OutAudio[SampleIndex] + OutAudio[SampleIndex + 1]);
-			CurrentEnvelopeValue = EnvelopeFollower.ProcessAudio(InputSample);
+			CurrentEnvelopeValue = 0.f;
 		}
 	}
 
diff --git a/Engine/Source/Runtime/AudioMixer/Private/AudioMixerSourceManager.cpp b/Engine/Source/Runtime/AudioMixer/Private/AudioMixerSourceManager.cpp
--- a/Engine/Source/Runtime/AudioMixer/Private/AudioMixerSourceManager.cpp
+++ b/Engine/Source/Runtime/AudioMixer/Private/AudioMixerSourceManager.cpp
@@ -245,7 +245,13 @@ namespace Audio
 
 			SourceInfo.SourceEffectChainId = INDEX_NONE;
 
-			SourceInfo.SourceEnvelopeFollower = Audio::FEnvelopeFollower(MixerDevice->SampleRate, 10, 100, Audio::EPeakMode::Peak);
+			Audio::FInlineEnvelopeFollowerInitParams EnvelopeFollowerInitParams;
+			EnvelopeFollowerInitParams.SampleRate = MixerDevice->SampleRate;
+			EnvelopeFollowerInitParams.AttackTimeMsec = 10.f;
+			EnvelopeFollowerInitParams.ReleaseTimeMsec = 100.f;
+			EnvelopeFollowerInitParams.Mode = EPeakMode::Peak;
+			SourceInfo.SourceEnvelopeFollower = Audio::FInlineEnvelopeFollower(EnvelopeFollowerInitParams);
+
 			SourceInfo.SourceEnvelopeValue = 0.0f;
 			SourceInfo.bEffectTailsDone = false;
 		
@@ -770,7 +776,12 @@ namespace Audio
 			SourceInfo.LowPassFilter.Init(MixerDevice->SampleRate, InitParams.NumInputChannels);
 			SourceInfo.HighPassFilter.Init(MixerDevice->SampleRate, InitParams.NumInputChannels);
 
-			SourceInfo.SourceEnvelopeFollower = Audio::FEnvelopeFollower(MixerDevice->SampleRate / NumOutputFrames, (float)InitParams.EnvelopeFollowerAttackTime, (float)InitParams.EnvelopeFollowerReleaseTime, Audio::EPeakMode::Peak);
+			Audio::FInlineEnvelopeFollowerInitParams EnvelopeFollowerInitParams;
+			EnvelopeFollowerInitParams.SampleRate = MixerDevice->SampleRate / NumOutputFrames;
+			EnvelopeFollowerInitParams.AttackTimeMsec = (float)InitParams.EnvelopeFollowerAttackTime;
+			EnvelopeFollowerInitParams.ReleaseTimeMsec = (float)InitParams.EnvelopeFollowerReleaseTime;
+			EnvelopeFollowerInitParams.Mode = EPeakMode::Peak;
+			SourceInfo.SourceEnvelopeFollower = Audio::FInlineEnvelopeFollower(EnvelopeFollowerInitParams);
 
 			SourceInfo.VolumeModulation = VolumeModulation;
 			SourceInfo.PitchModulation = PitchModulation;
@@ -2310,10 +2321,8 @@ namespace Audio
 			{
 				// Compute the source envelope using pre-distance attenuation buffer
 				float AverageSampleValue = Audio::BufferGetAverageAbsValue(PreDistanceAttenBufferPtr, NumSamples);
-				SourceInfo.SourceEnvelopeFollower.ProcessAudio(AverageSampleValue);
-
-				// Copy the current value of the envelope follower (block-rate value)
-				SourceInfo.SourceEnvelopeValue = SourceInfo.SourceEnvelopeFollower.GetCurrentValue();
+				SourceInfo.SourceEnvelopeValue = SourceInfo.SourceEnvelopeFollower.ProcessSample(AverageSampleValue);
+				SourceInfo.SourceEnvelopeValue = FMath::Clamp(SourceInfo.SourceEnvelopeValue, 0.f, 1.f);
 
 				SourceInfo.bEffectTailsDone = SourceInfo.bEffectTailsDone || SourceInfo.SourceEnvelopeValue < ENVELOPE_TAIL_THRESHOLD;
 			}
diff --git a/Engine/Source/Runtime/AudioMixer/Private/AudioMixerSourceManager.h b/Engine/Source/Runtime/AudioMixer/Private/AudioMixerSourceManager.h
--- a/Engine/Source/Runtime/AudioMixer/Private/AudioMixerSourceManager.h
+++ b/Engine/Source/Runtime/AudioMixer/Private/AudioMixerSourceManager.h
@@ -409,7 +409,7 @@ namespace Audio
 			FAudioPluginSourceOutputData AudioPluginOutputData;
 
 			// A DSP object which tracks the amplitude envelope of a source.
-			Audio::FEnvelopeFollower SourceEnvelopeFollower;
+			Audio::FInlineEnvelopeFollower SourceEnvelopeFollower;
 			float SourceEnvelopeValue;
 
 			// Modulation destinations
diff --git a/Engine/Source/Runtime/AudioMixer/Private/AudioMixerSubmix.cpp b/Engine/Source/Runtime/AudioMixer/Private/AudioMixerSubmix.cpp
--- a/Engine/Source/Runtime/AudioMixer/Private/AudioMixerSubmix.cpp
+++ b/Engine/Source/Runtime/AudioMixer/Private/AudioMixerSubmix.cpp
@@ -181,8 +181,6 @@ namespace Audio
 		, bIsBackgroundMuted(false)
 		, bIsSpectrumAnalyzing(false)
 	{
-		EnvelopeFollowers.Reset();
-		EnvelopeFollowers.AddDefaulted(AUDIO_MIXER_MAX_OUTPUT_CHANNELS);
 	}
 
 	FMixerSubmix::~FMixerSubmix()
@@ -1315,19 +1313,21 @@ namespace Audio
 			FScopeLock EnvelopeScopeLock(&EnvelopeCriticalSection);
 			FMemory::Memset(EnvelopeValues, sizeof(float) * AUDIO_MIXER_MAX_OUTPUT_CHANNELS);
 
-			for (int32 ChannelIndex = 0; ChannelIndex < NumChannels; ++ChannelIndex)
+			if (NumChannels > 0)
 			{
-				// Get the envelope follower for the channel
-				FEnvelopeFollower& EnvFollower = EnvelopeFollowers[ChannelIndex];
-
-				// Track the last sample
-				for (int32 SampleIndex = ChannelIndex; SampleIndex < BufferSamples; SampleIndex += NumChannels)
+				const int32 NumFrames = BufferSamples / NumChannels;
+				if (EnvelopeFollower.GetNumChannels() != NumChannels)
 				{
-					const float SampleValue = AudioBufferPtr[SampleIndex];
-					EnvFollower.ProcessAudio(SampleValue);
+					EnvelopeFollower.SetNumChannels(NumChannels);
 				}
 
-				EnvelopeValues[ChannelIndex] = EnvFollower.GetCurrentValue();
+				EnvelopeFollower.ProcessAudio(AudioBufferPtr, NumFrames);
+				const TArray<float>& EnvValues = EnvelopeFollower.GetEnvelopeValues();
+
+				check(EnvValues.Num() == NumChannels);
+
+				FMemory::Memcpy(EnvelopeValues, EnvValues.GetData(), sizeof(float) * NumChannels);
+				Audio::ArrayClampInPlace(MakeArrayView(EnvelopeValues, NumChannels), 0.f, 1.f);
 			}
 
 			EnvelopeNumChannels = NumChannels;
@@ -1903,11 +1903,17 @@ namespace Audio
 	{
 		if (!bIsEnvelopeFollowing)
 		{
+			FEnvelopeFollowerInitParams EnvelopeFollowerInitParams;
+			EnvelopeFollowerInitParams.SampleRate = GetSampleRate(); 
+			EnvelopeFollowerInitParams.NumChannels = NumChannels; 
+			EnvelopeFollowerInitParams.AttackTimeMsec = static_cast<float>(AttackTime);
+			EnvelopeFollowerInitParams.ReleaseTimeMsec = static_cast<float>(ReleaseTime);
+			EnvelopeFollower.Init(EnvelopeFollowerInitParams);
+
 			// Zero out any previous envelope values which may have been in the array before starting up
 			for (int32 ChannelIndex = 0; ChannelIndex < AUDIO_MIXER_MAX_OUTPUT_CHANNELS; ++ChannelIndex)
 			{
 				EnvelopeValues[ChannelIndex] = 0.0f;
-				EnvelopeFollowers[ChannelIndex].Init(GetSampleRate(), AttackTime, ReleaseTime);
 			}
 
 			bIsEnvelopeFollowing = true;
@@ -2010,7 +2016,12 @@ namespace Audio
 						DelegateInfo.SpectrumBandExtractor->AddBand(NewExtractorBandSettings);
 
 						FSpectralAnalysisBandInfo NewBand;
-						NewBand.EnvelopeFollower.Init(DelegateInfo.DelegateSettings.UpdateRate, BandSettings.AttackTimeMsec, BandSettings.ReleaseTimeMsec);
+
+						FInlineEnvelopeFollowerInitParams EnvelopeFollowerInitParams;
+						EnvelopeFollowerInitParams.SampleRate = DelegateInfo.DelegateSettings.UpdateRate;
+						EnvelopeFollowerInitParams.AttackTimeMsec = static_cast<float>(BandSettings.AttackTimeMsec);
+						EnvelopeFollowerInitParams.ReleaseTimeMsec = static_cast<float>(BandSettings.ReleaseTimeMsec);
+						NewBand.EnvelopeFollower.Init(EnvelopeFollowerInitParams);
 					
 						DelegateInfo.SpectralBands.Add(NewBand);
 					}
@@ -2219,7 +2230,8 @@ namespace Audio
 							if (ensure(ResultIndex < DelegateInfo.SpectralBands.Num()))
 							{
 								FSpectralAnalysisBandInfo& BandInfo = DelegateInfo.SpectralBands[ResultIndex];
-								SpectralResults[ResultIndex] = BandInfo.EnvelopeFollower.ProcessAudioNonClamped(SpectralResults[ResultIndex]);
+
+								SpectralResults[ResultIndex] = BandInfo.EnvelopeFollower.ProcessSample(SpectralResults[ResultIndex]);
 							}
 						}
 
diff --git a/Engine/Source/Runtime/AudioMixer/Public/AudioMixerSubmix.h b/Engine/Source/Runtime/AudioMixer/Public/AudioMixerSubmix.h
--- a/Engine/Source/Runtime/AudioMixer/Public/AudioMixerSubmix.h
+++ b/Engine/Source/Runtime/AudioMixer/Public/AudioMixerSubmix.h
@@ -522,7 +522,7 @@ namespace Audio
 
 		// Envelope following data
 		float EnvelopeValues[AUDIO_MIXER_MAX_OUTPUT_CHANNELS];
-		TArray<Audio::FEnvelopeFollower> EnvelopeFollowers;
+		Audio::FEnvelopeFollower EnvelopeFollower;
 		int32 EnvelopeNumChannels;
 		FCriticalSection EnvelopeCriticalSection;
 
@@ -566,7 +566,7 @@ namespace Audio
 
 		struct FSpectralAnalysisBandInfo
 		{
-			FEnvelopeFollower EnvelopeFollower;
+			FInlineEnvelopeFollower EnvelopeFollower;
 		};
 
 		struct FSpectrumAnalysisDelegateInfo
diff --git a/Engine/Source/Runtime/Engine/Private/SoundWave.cpp b/Engine/Source/Runtime/Engine/Private/SoundWave.cpp
--- a/Engine/Source/Runtime/Engine/Private/SoundWave.cpp
+++ b/Engine/Source/Runtime/Engine/Private/SoundWave.cpp
@@ -1617,12 +1617,15 @@ void USoundWave::BakeFFTAnalysis()
 		}
 
 		// Prepare the spectral envelope followers
-		TArray<Audio::FEnvelopeFollower> SpectralEnvelopeFollowers;
-		SpectralEnvelopeFollowers.AddDefaulted(FrequenciesToAnalyze.Num());
+		Audio::FInlineEnvelopeFollowerInitParams EnvelopeFollowerInitParams;
+		EnvelopeFollowerInitParams.SampleRate = static_cast<float>(RawDataSampleRate) / static_cast<float>(FMath::Max(1, FFTAnalysisFrameSize));
+		EnvelopeFollowerInitParams.AttackTimeMsec = static_cast<float>(FFTAnalysisAttackTime);
+		EnvelopeFollowerInitParams.ReleaseTimeMsec = static_cast<float>(FFTAnalysisReleaseTime);
 
-		for (Audio::FEnvelopeFollower& EnvFollower : SpectralEnvelopeFollowers)
+		TArray<Audio::FInlineEnvelopeFollower> SpectralEnvelopeFollowers;
+		for (int32 i = 0; i < FrequenciesToAnalyze.Num(); i++)
 		{
-			EnvFollower.Init((float)RawDataSampleRate / FFTAnalysisFrameSize, FFTAnalysisAttackTime, FFTAnalysisReleaseTime);
+			SpectralEnvelopeFollowers.Emplace(EnvelopeFollowerInitParams);
 		}
 
 		// Build a new spectrum analyzer
@@ -1668,7 +1671,7 @@ void USoundWave::BakeFFTAnalysis()
 					DataEntry.Magnitude = SpectrumAnalyzer.GetMagnitudeForFrequency(Frequency);
 
 					// Feed the magnitude through the spectral envelope follower for this band
-					DataEntry.Magnitude = SpectralEnvelopeFollowers[Index].ProcessAudioNonClamped(DataEntry.Magnitude);
+					DataEntry.Magnitude = SpectralEnvelopeFollowers[Index].ProcessSample(DataEntry.Magnitude);
 
 					// Track the max magnitude so we can later set normalized magnitudes
 					if (DataEntry.Magnitude > MaximumMagnitude)
@@ -1784,8 +1787,11 @@ void USoundWave::BakeEnvelopeAnalysis()
 		const uint32 NumFrames = (RawImportedWaveData.Num() / sizeof(int16)) / RawDataNumChannels;
 		int16* InputData = (int16*)RawImportedWaveData.GetData();
 
-		Audio::FEnvelopeFollower EnvelopeFollower;
-		EnvelopeFollower.Init(RawDataSampleRate, EnvelopeFollowerAttackTime, EnvelopeFollowerReleaseTime);
+		Audio::FInlineEnvelopeFollowerInitParams EnvelopeFollowerInitParams;
+		EnvelopeFollowerInitParams.SampleRate = RawDataSampleRate; 
+		EnvelopeFollowerInitParams.AttackTimeMsec = static_cast<float>(EnvelopeFollowerAttackTime);
+		EnvelopeFollowerInitParams.ReleaseTimeMsec = static_cast<float>(EnvelopeFollowerReleaseTime);
+		Audio::FInlineEnvelopeFollower EnvelopeFollower(EnvelopeFollowerInitParams);
 
 		for (uint32 FrameIndex = 0; FrameIndex < NumFrames; ++FrameIndex)
 		{
@@ -1797,7 +1803,10 @@ void USoundWave::BakeEnvelopeAnalysis()
 			}
 			SampleValue /= RawDataNumChannels;
 
-			float Output = EnvelopeFollower.ProcessAudio(SampleValue);
+			float Output = EnvelopeFollower.ProcessSample(SampleValue);
+			Output = FMath::Clamp(Output, 0.f, 1.f);
+			
+
 
 			// Until we reached the frame size
 			if (FrameIndex % EnvelopeFollowerFrameSize == 0)
diff --git a/Engine/Source/Runtime/MediaAssets/Private/Assets/MediaSoundComponent.cpp b/Engine/Source/Runtime/MediaAssets/Private/Assets/MediaSoundComponent.cpp
--- a/Engine/Source/Runtime/MediaAssets/Private/Assets/MediaSoundComponent.cpp
+++ b/Engine/Source/Runtime/MediaAssets/Private/Assets/MediaSoundComponent.cpp
@@ -343,6 +343,11 @@ bool UMediaSoundComponent::Init(int32& SampleRate)
 
 	Resampler->Initialize(NumChannels, SampleRate);
 
+	Audio::FEnvelopeFollowerInitParams EnvelopeInitParams;
+	EnvelopeInitParams.SampleRate = SampleRate;
+	EnvelopeInitParams.NumChannels = 1; //EnvelopeFollower uses mixed down mono buffer 
+	EnvelopeFollower.Init(EnvelopeInitParams);
+
 	return true;
 }
 
@@ -432,12 +437,17 @@ int32 UMediaSoundComponent::OnGenerateAudio(float* OutAudio, int32 NumSamples)
 						bEnvelopeFollowerSettingsChanged = false;
 					}
 
-					for (int32 FrameIndex = 0; FrameIndex < NumFrames; ++FrameIndex)
+					EnvelopeFollower.ProcessAudio(BufferToUseForAnalysis, NumFrames);
+
+					const TArray<float>& EnvelopeValues = EnvelopeFollower.GetEnvelopeValues();
+					if (ensure(EnvelopeValues.Num() > 0))
 					{
-						EnvelopeFollower.ProcessAudio(BufferToUseForAnalysis[FrameIndex]);
+						CurrentEnvelopeValue = FMath::Clamp(EnvelopeValues[0], 0.f, 1.f);
+					}
+					else
+					{
+						CurrentEnvelopeValue = 0.f;
 					}
-
-					CurrentEnvelopeValue = EnvelopeFollower.GetCurrentValue();
 				}
 			}
 		}
diff --git a/Engine/Source/Runtime/Online/Voice/Private/Windows/VoiceCaptureWindows.cpp b/Engine/Source/Runtime/Online/Voice/Private/Windows/VoiceCaptureWindows.cpp
--- a/Engine/Source/Runtime/Online/Voice/Private/Windows/VoiceCaptureWindows.cpp
+++ b/Engine/Source/Runtime/Online/Voice/Private/Windows/VoiceCaptureWindows.cpp
@@ -123,7 +123,14 @@ bool FVoiceCaptureWindows::Init(const FString& DeviceName, int32 SampleRate, int
 	static IConsoleVariable* SilenceDetectionReleaseCVar = IConsoleManager::Get().FindConsoleVariable(TEXT("voice.SilenceDetectionReleaseTime"));
 	check(SilenceDetectionReleaseCVar);				
 
-	MicLevelDetector.Init(SampleRate, SilenceDetectionAttackCVar->GetFloat(), SilenceDetectionReleaseCVar->GetFloat(), MicSilenceDetectionConfig::LevelDetectionMode, MicSilenceDetectionConfig::IsAnalog);
+	Audio::FInlineEnvelopeFollowerInitParams EnvelopeFollowerInitParams;
+	EnvelopeFollowerInitParams.SampleRate = SampleRate;
+	EnvelopeFollowerInitParams.AttackTimeMsec = SilenceDetectionAttackCVar->GetFloat();
+	EnvelopeFollowerInitParams.ReleaseTimeMsec = SilenceDetectionReleaseCVar->GetFloat();
+	EnvelopeFollowerInitParams.Mode = Audio::EPeakMode::Peak;
+	EnvelopeFollowerInitParams.bIsAnalog = MicSilenceDetectionConfig::IsAnalog;
+
+	MicLevelDetector.Init(EnvelopeFollowerInitParams);
 	
 	const int32 AttackInSamples = SampleRate * SilenceDetectionAttackCVar->GetFloat() * 0.001f;
 	LookaheadBuffer.Init(AttackInSamples + 1);
@@ -476,13 +483,14 @@ void FVoiceCaptureWindows::ProcessData()
 					Temp += InputBuffer[FrameIndex + ChannelIndex];
 				}
 
-				MicLevelDetector.ProcessAudio(Temp);
+				
+				float Envelope = MicLevelDetector.ProcessSample(Temp / 32768.f);
 				LookaheadBuffer.ProcessSample(Temp, Temp);
 
-				bIsMicActive = MicLevelDetector.GetCurrentValue() > MicSilenceThreshold;
+				bIsMicActive = Envelope > MicSilenceThreshold;
 
 				// If we have just crossed the noise gate threshold, begin interpoloating to 1.0 or 0.0
-				const bool bIsMicAboveNoiseGateThreshold = MicLevelDetector.GetCurrentValue() > MicNoiseGateThreshold;
+				const bool bIsMicAboveNoiseGateThreshold = Envelope > MicNoiseGateThreshold;
 
 				if (bIsMicAboveNoiseGateThreshold && !bWasMicAboveNoiseGateThreshold)
 				{
@@ -555,13 +563,14 @@ void FVoiceCaptureWindows::ProcessData()
 					Temp += InputBuffer[FrameIndex + ChannelIndex];
 				}
 				
-				MicLevelDetector.ProcessAudio(Temp);
+				
+				float Envelope = MicLevelDetector.ProcessSample(static_cast<float>(Temp) / 32768.f);
 				LookaheadBuffer.ProcessSample(Temp, Temp);
 
-				bIsMicActive = MicLevelDetector.GetCurrentValue() > MicSilenceThreshold;
+				bIsMicActive = Envelope > MicSilenceThreshold;
 
 				// If we have just crossed the noise gate threshold, begin interpoloating to 1.0 or 0.0
-				const bool bIsMicAboveNoiseGateThreshold = MicLevelDetector.GetCurrentValue() > MicNoiseGateThreshold;
+				const bool bIsMicAboveNoiseGateThreshold = Envelope > MicNoiseGateThreshold;
 
 				if (bIsMicAboveNoiseGateThreshold && !bWasMicAboveNoiseGateThreshold)
 				{
@@ -713,10 +722,10 @@ EVoiceCaptureState::Type FVoiceCaptureWindows::GetVoiceData(uint8* OutVoiceBuffe
 
 		if (FPlatformTime::Seconds() - TimeLastPrinted > AmplitudeStringDisplayRate)
 		{
-			const float MicLevel = MicLevelDetector.GetCurrentValue();
+			const float MicLevel = MicLevelDetector.GetValue();
 			FString PrintString = FString::Printf(TEXT("Mic Amp: %.2f"), MicLevel);
 
-			int32 NumTicks = FMath::FloorToInt(MicLevelDetector.GetCurrentValue() * TotalNumTicks);
+			int32 NumTicks = FMath::FloorToInt(MicLevel * TotalNumTicks);
 
 			for (int32 Iteration = 0; Iteration < NumTicks; Iteration++)
 			{
@@ -819,7 +828,7 @@ bool FVoiceCaptureWindows::Tick(float DeltaTime)
 
 float FVoiceCaptureWindows::GetCurrentAmplitude() const
 {
-	return MicLevelDetector.GetCurrentValue();
+	return MicLevelDetector.GetValue();
 }
 
 void FVoiceCaptureWindows::DumpState() const
diff --git a/Engine/Source/Runtime/Online/Voice/Private/Windows/VoiceCaptureWindows.h b/Engine/Source/Runtime/Online/Voice/Private/Windows/VoiceCaptureWindows.h
--- a/Engine/Source/Runtime/Online/Voice/Private/Windows/VoiceCaptureWindows.h
+++ b/Engine/Source/Runtime/Online/Voice/Private/Windows/VoiceCaptureWindows.h
@@ -169,7 +169,7 @@ private:
 	* Envelope following DSP object.
 	* Configured using the MicSilenceDetectionConfig namespace in VoiceConfig.h and the CVars defined in VoiceConfig.cpp.
 	*/
-	Audio::FEnvelopeFollower MicLevelDetector;
+	Audio::FInlineEnvelopeFollower MicLevelDetector;
 	/*
 	* Whether the microphone level is above the threshold set 
 	*/
diff --git a/Engine/Source/Runtime/SignalProcessing/Private/DynamicsProcessor.cpp b/Engine/Source/Runtime/SignalProcessing/Private/DynamicsProcessor.cpp
--- a/Engine/Source/Runtime/SignalProcessing/Private/DynamicsProcessor.cpp
+++ b/Engine/Source/Runtime/SignalProcessing/Private/DynamicsProcessor.cpp
@@ -1,12 +1,13 @@
 // Copyright Epic Games, Inc. All Rights Reserved.
 
 #include "DSP/DynamicsProcessor.h"
+#include "SignalProcessingModule.h"
 
 namespace Audio
 {
 	FDynamicsProcessor::FDynamicsProcessor()
 		: ProcessingMode(EDynamicsProcessingMode::Compressor)
-		, EnvelopeFollowerPeakMode(EPeakMode::RootMeanSquared)
+		, EnvelopeFollowerPeakMode(EPeakMode::Peak)
 		, LookaheedDelayMsec(10.0f)
 		, AttackTimeMsec(20.0f)
 		, ReleaseTimeMsec(1000.0f)
@@ -45,7 +46,7 @@ namespace Audio
 			LookaheadDelay[Channel].Init(SampleRate, 0.1f);
 			LookaheadDelay[Channel].SetDelayMsec(LookaheedDelayMsec);
 
-			EnvFollower[Channel].Init(SampleRate, AttackTimeMsec, ReleaseTimeMsec, EPeakMode::RootMeanSquared, bIsAnalogMode);
+			EnvFollower[Channel].Init(FInlineEnvelopeFollowerInitParams{SampleRate, AttackTimeMsec, ReleaseTimeMsec, EnvelopeFollowerPeakMode, bIsAnalogMode});
 		}
 
 		InputLowshelfFilter.Init(SampleRate, InNumChannels, EBiquadFilter::LowShelf);
@@ -161,11 +162,10 @@ namespace Audio
 		if (InNumChannels != EnvFollower.Num())
 		{
 			EnvFollower.Reset();
-			EnvFollower.AddDefaulted(InNumChannels);
 
 			for (int32 Channel = 0; Channel < InNumChannels; ++Channel)
 			{
-				EnvFollower[Channel].Init(SampleRate, AttackTimeMsec, ReleaseTimeMsec, EnvelopeFollowerPeakMode, bIsAnalogMode);
+				EnvFollower.Emplace(FInlineEnvelopeFollowerInitParams{SampleRate, AttackTimeMsec, ReleaseTimeMsec, EnvelopeFollowerPeakMode, bIsAnalogMode});
 			}
 		}
 
@@ -231,7 +231,7 @@ namespace Audio
 		EnvelopeFollowerPeakMode = InEnvelopeFollowerModeType;
 		for (int32 Channel = 0; Channel < EnvFollower.Num(); ++Channel)
 		{
-			EnvFollower[Channel].SetMode(InEnvelopeFollowerModeType);
+			EnvFollower[Channel].SetMode(EnvelopeFollowerPeakMode);
 		}
 	}
 
@@ -383,8 +383,9 @@ namespace Audio
 
 		for (int32 Channel = 0; Channel < KeyNumChannels; ++Channel)
 		{
-			DetectorOuts[Channel] = EnvFollower[Channel].ProcessAudio(DetectorGain * KeyIn[Channel]);
+			DetectorOuts[Channel] = EnvFollower[Channel].ProcessSample(DetectorGain * KeyIn[Channel]);
 		}
+		ArrayClampInPlace(DetectorOuts, 0.f, 1.f);
 
 		switch (LinkMode)
 		{
diff --git a/Engine/Source/Runtime/SignalProcessing/Private/EnvelopeFollower.cpp b/Engine/Source/Runtime/SignalProcessing/Private/EnvelopeFollower.cpp
--- a/Engine/Source/Runtime/SignalProcessing/Private/EnvelopeFollower.cpp
+++ b/Engine/Source/Runtime/SignalProcessing/Private/EnvelopeFollower.cpp
@@ -1,224 +1,463 @@
 // Copyright Epic Games, Inc. All Rights Reserved.
 
 #include "DSP/EnvelopeFollower.h"
-#include "DSP/Dsp.h"
-#include "DSP/BufferVectorOperations.h"
+#include "DSP/FloatArrayMath.h"
 
 namespace Audio
 {
-	// see https://en.wikipedia.org/wiki/RC_time_constant
-	// Time constants indicate how quickly the envelope follower responds to changes in input
-	static const float AnalogTImeConstant = 1.00239343f;
-	static const float DigitalTimeConstant = 4.60517019f;
+	FAttackRelease::FAttackRelease(float InSampleRate, float InAttackTimeMsec, float InReleaseTimeMsec, bool bInIsAnalog)
+	: SampleRate(InSampleRate)
+	, AttackTimeSamples(1.f)
+	, ReleaseTimeSamples(1.f)
+	, bIsAnalog(bInIsAnalog)
+	{
+		if (!ensure(SampleRate > 0.f))
+		{
+			SampleRate = 48000.f;
+		}
+		// Set the attack and release times using the default values
+		SetAttackTime(InAttackTimeMsec);
+		SetReleaseTime(InReleaseTimeMsec);
+	}
 
-	FEnvelopeFollower::FEnvelopeFollower()
-		: EnvMode(EPeakMode::Peak)
-		, MeanWindowSize(DefaultWindowSize)
-		, MeanHopSize(DefaultHopSize)
-		, SampleRate(44100.0f)
-		, AttackTimeMsec(0.0f)
-		, AttackTimeSamples(0.0f)
-		, ReleaseTimeMsec(0.0f)
-		, ReleaseTimeSamples(0.0f)
-		, CurrentEnvelopeValue(0.0f)
-		, bIsAnalog(true)
+	void FAttackRelease::SetSampleRate(float InSampleRate)
+	{
+		if (ensure(InSampleRate > 0.f))
+		{
+			SampleRate = InSampleRate;
+			SetAttackTime(AttackTimeMsec);
+			SetReleaseTime(ReleaseTimeMsec);
+		}
+	}
+
+	void FAttackRelease::SetAnalog(bool bInIsAnalog)
 	{
+		bIsAnalog = bInIsAnalog;
+		SetAttackTime(AttackTimeMsec);
+		SetReleaseTime(ReleaseTimeMsec);
 	}
 
-	FEnvelopeFollower::FEnvelopeFollower(const float InSampleRate,
-		const float InAttackTimeMsec,
-		const float InReleaseTimeMSec,
-		const EPeakMode::Type InMode,
-		const bool bInIsAnalog,
-		const int32 InWindowSizeForMean,
-		const int32 InHopSizeForMean)		
+	void FAttackRelease::SetAttackTime(float InAttackTimeMsec)
 	{
-		Init(InSampleRate, InAttackTimeMsec, InReleaseTimeMSec, InMode, bInIsAnalog, InWindowSizeForMean, InHopSizeForMean);
+		AttackTimeMsec = InAttackTimeMsec;
+		if (AttackTimeMsec > 0.f)
+		{
+			float TimeConstant = bIsAnalog ? AnalogTimeConstant : DigitalTimeConstant;
+			AttackTimeSamples = FMath::Exp(-1000.0f * TimeConstant / (AttackTimeMsec * SampleRate));
+		}
+		else
+		{
+			AttackTimeSamples = 0.f;
+		}
 	}
 
-	FEnvelopeFollower::~FEnvelopeFollower()
+	void FAttackRelease::SetReleaseTime(float InReleaseTimeMsec)
 	{
+		ReleaseTimeMsec = InReleaseTimeMsec;
+		if (ReleaseTimeMsec > 0.f)
+		{
+			float TimeConstant = bIsAnalog ? AnalogTimeConstant : DigitalTimeConstant;
+			ReleaseTimeSamples = FMath::Exp(-1000.0f * TimeConstant / (ReleaseTimeMsec * SampleRate));
+		}
+		else
+		{
+			ReleaseTimeSamples = 0.f;
+		}
 	}
 
-	void FEnvelopeFollower::Init(const float InSampleRate, 
-		const float InAttackTimeMsec, 
-		const float InReleaseTimeMSec, 
-		const EPeakMode::Type InMode, 
-		const bool bInIsAnalog,
-		const int32 InWindowSizeForMean,
-		const int32 InHopSizeForMean)
+	FAttackReleaseSmoother::FAttackReleaseSmoother(float InSampleRate, int32 InNumChannels, float InAttackTimeMsec, float InReleaseTimeMsec, bool bInIsAnalog)
+	: FAttackRelease(InSampleRate, InAttackTimeMsec, InReleaseTimeMsec, bInIsAnalog)
+	, NumChannels(0)
 	{
-		check(InWindowSizeForMean > 0);
-		check(InHopSizeForMean > 0);
-		check(InWindowSizeForMean > InHopSizeForMean);
+		SetNumChannels(InNumChannels);
+	}
 
-		SampleRate = InSampleRate;
+	void FAttackReleaseSmoother::ProcessAudio(const float* InBuffer, int32 InNumFrames)
+	{
+		if (InNumFrames > 0)
+		{
+			const int32 NumSamples = InNumFrames * NumChannels;
+			const float AttackSamples = GetAttackTimeSamples();
+			const float ReleaseSamples = GetReleaseTimeSamples();
 
-		bIsAnalog = bInIsAnalog;
-		EnvMode = InMode;
+			for (int32 Channel = 0; Channel < NumChannels; Channel++)
+			{
+				float EnvelopeValue = PriorEnvelopeValues[Channel];
 
-		MeanWindowSize = InWindowSizeForMean;
-		MeanHopSize = InHopSizeForMean;
+				for (int32 Pos = Channel; Pos < NumSamples; Pos += NumChannels)
+				{
+					float Diff = EnvelopeValue - InBuffer[Pos];
+					if (Diff > 0.f)
+					{
+						EnvelopeValue = (AttackSamples * Diff) + InBuffer[Pos];
+					}
+					else
+					{
+						EnvelopeValue = (ReleaseSamples * Diff) + InBuffer[Pos];
+					}
 
-		// Set the attack and release times using the default values
-		SetAttackTime(InAttackTimeMsec);
-		SetReleaseTime(InReleaseTimeMSec);
+					EnvelopeValue = Audio::UnderflowClamp(EnvelopeValue);
+				}
 
-		ScratchBuffer.SetNum(InWindowSizeForMean);
+				PriorEnvelopeValues[Channel] = EnvelopeValue;
+			}
+		}
 	}
 
-	void FEnvelopeFollower::Reset()
+	void FAttackReleaseSmoother::ProcessAudio(const float* InBuffer, int32 InNumFrames, float* OutBuffer)
 	{
-		CurrentEnvelopeValue = 0.0f;
+		if (InNumFrames > 0)
+		{
+			const int32 NumSamples = InNumFrames * NumChannels;
+			const float AttackSamples = GetAttackTimeSamples();
+			const float ReleaseSamples = GetReleaseTimeSamples();
+
+			for (int32 Channel = 0; Channel < NumChannels; Channel++)
+			{
+				float EnvelopeValue = PriorEnvelopeValues[Channel];
+
+				for (int32 Pos = Channel; Pos < NumSamples; Pos += NumChannels)
+				{
+					float Diff = EnvelopeValue - InBuffer[Pos];
+					if (Diff <= 0.f)
+					{
+						EnvelopeValue = (AttackSamples * Diff) + InBuffer[Pos];
+					}
+					else
+					{
+						EnvelopeValue = (ReleaseSamples * Diff) + InBuffer[Pos];
+					}
+
+					EnvelopeValue = Audio::UnderflowClamp(EnvelopeValue);
+					OutBuffer[Pos] = EnvelopeValue;
+				}
+
+				PriorEnvelopeValues[Channel] = EnvelopeValue;
+			}
+		}
 	}
 
-	void FEnvelopeFollower::SetAnalog(const bool bInIsAnalog)
+	const TArray<float>& FAttackReleaseSmoother::GetEnvelopeValues() const
 	{
-		bIsAnalog = bInIsAnalog;
-		SetAttackTime(AttackTimeMsec);
-		SetReleaseTime(ReleaseTimeMsec);
+		return PriorEnvelopeValues;
 	}
-
-	void FEnvelopeFollower::SetAttackTime(const float InAttackTimeMsec)
+	void FAttackReleaseSmoother::SetNumChannels(int32 InNumChannels)
 	{
-		AttackTimeMsec = InAttackTimeMsec;
-		const float TimeConstant = bIsAnalog ? AnalogTImeConstant : DigitalTimeConstant;
-		AttackTimeSamples = FMath::Exp(-1000.0f * TimeConstant / (AttackTimeMsec * SampleRate));
+		if (InNumChannels != NumChannels)
+		{
+			PriorEnvelopeValues.Reset(InNumChannels);
+			if (ensure(InNumChannels > 0))
+			{
+				PriorEnvelopeValues.AddZeroed(InNumChannels);
+			}
+			NumChannels = InNumChannels;
+		}
 	}
 
-	void FEnvelopeFollower::SetReleaseTime(const float InReleaseTimeMsec)
+	void FAttackReleaseSmoother::Reset()
 	{
-		ReleaseTimeMsec = InReleaseTimeMsec;
-		const float TimeConstant = bIsAnalog ? AnalogTImeConstant : DigitalTimeConstant;
-		ReleaseTimeSamples = FMath::Exp(-1000.0f * TimeConstant / (InReleaseTimeMsec * SampleRate));
+		if (NumChannels > 0)
+		{
+			FMemory::Memset(PriorEnvelopeValues.GetData(), 0, sizeof(float) * NumChannels);
+		}
 	}
 
-	void FEnvelopeFollower::SetMode(const EPeakMode::Type InMode)
+
+	FMeanSquaredFIR::FMeanSquaredFIR(float InSampleRate, int32 InNumChannels, float InWindowTimeMsec)
+	: SampleRate(InSampleRate)
+	, NumChannels(InNumChannels)
+	, WindowTimeSamples(256)
+	, NormFactor(1.f)
 	{
-		EnvMode = InMode;
+		if (!ensure(SampleRate > 0.f))
+		{
+			SampleRate = 48000.f;
+		}
+
+		if (ensure(NumChannels > 0))
+		{
+			ChannelValues.AddZeroed(NumChannels);
+		}
+
+		SetWindowSize(InWindowTimeMsec);
 	}
 
-	float FEnvelopeFollower::ProcessAudio(const float InAudioSample)
+	void FMeanSquaredFIR::SetWindowSize(float InWindowTimeMsec)
 	{
-		ProcessAudio(&InAudioSample, 1);
+		if (ensure(InWindowTimeMsec > 0.f))
+		{
+			WindowTimeFrames = FMath::Max(1, FMath::RoundToInt(InWindowTimeMsec * 1000.f / SampleRate));
+			NormFactor = 1.f / static_cast<float>(WindowTimeFrames);
+			WindowTimeSamples = WindowTimeFrames * NumChannels;
+		}
 
-		// Update and return the envelope value
-		return CurrentEnvelopeValue;
+		Reset();
 	}
 
-	float FEnvelopeFollower::ProcessAudio(const float* InAudioBuffer, int32 InNumSamples)
+	void FMeanSquaredFIR::Reset()
 	{
-		// MS/RMS
-		if (EnvMode == EPeakMode::MeanSquared || EnvMode == EPeakMode::RootMeanSquared)
+		if (NumChannels > 0)
 		{
-			TCircularAudioBuffer<float> SumBuffer(MeanWindowSize);
-			int32 SampleIndex = 0;
-			// This loop iterates per-window, not per-sample
-			while (SampleIndex < InNumSamples)
-			{
-				// Get a new Window
-				const int32 SamplesToCopy = MeanHopSize < (InNumSamples - SampleIndex) ? MeanHopSize : (InNumSamples - SampleIndex);
-				SumBuffer.Push(InAudioBuffer + SampleIndex, SamplesToCopy);
-				float* CurrentWindow = ScratchBuffer.GetData();
-				SumBuffer.Pop(CurrentWindow, SamplesToCopy);
-				
-				// Calculate MS/RMS
-				float CurrentMean;
-				ArrayMeanSquared(TArrayView<float>(CurrentWindow, SamplesToCopy), CurrentMean);
-				if (EnvMode == EPeakMode::RootMeanSquared)
-				{
-					CurrentMean = FMath::Sqrt(CurrentMean);
-				}
+			ChannelValues.Reset();
+			ChannelValues.AddZeroed(NumChannels);
 
-				// Use MS/RMS to drive envelope following
-				for (int i = 0; i < SamplesToCopy; ++i)
-				{
-					ProcessAudioNonClamped(CurrentMean);
-				}
-				SampleIndex += SamplesToCopy;
+			WindowTimeSamples = WindowTimeFrames * NumChannels;
+			HistorySquared.Reset(FMath::Max(2 * WindowTimeSamples, DefaultHistoryCapacity));
+			if (WindowTimeSamples > 0)
+			{
+				HistorySquared.PushZeros(WindowTimeSamples);
 			}
 		}
-		// Peak mode
-		else
+	}
+
+	void FMeanSquaredFIR::SetNumChannels(int32 InNumChannels)
+	{
+		if (NumChannels != InNumChannels)
+		{
+			NumChannels = InNumChannels;
+			Reset(); // Update channel buffers and window buffers. 
+		}
+	}
+
+	void FMeanSquaredFIR::ProcessAudio(const float* InBuffer, int32 InNumFrames, float* OutBuffer)
+	{
+		const int32 NumSamples = InNumFrames * NumChannels;
+		SquaredHistoryBuffer.Reset(NumSamples);
+		SquaredInputBuffer.Reset(NumSamples);
+
+		if (NumSamples > 0)
 		{
-			for (int32 SampleIndex = 0; SampleIndex < InNumSamples; ++SampleIndex)
+			// Square the input data
+			SquaredInputBuffer.AddUninitialized(NumSamples);
+			ArraySquare(MakeArrayView(InBuffer, NumSamples), SquaredInputBuffer);	
+
+			// Save the squared data for later
+			HistorySquared.Push(SquaredInputBuffer);
+
+			// Get the historic input so it can be subtracted from 
+			// the accumulated value as it leaves the window.
+			SquaredHistoryBuffer.AddUninitialized(NumSamples);
+			HistorySquared.Pop(SquaredHistoryBuffer.GetData(), NumSamples);
+		}
+
+		const float* HistorySquaredData = SquaredHistoryBuffer.GetData();
+		const float* InputSquaredData = SquaredInputBuffer.GetData();
+		float* ValueData = ChannelValues.GetData();
+
+		// MS[n] = MS[n - 1] + x^2[n] - x^2[n - N]
+		for (int32 ChannelIndex = 0; ChannelIndex < NumChannels; ChannelIndex++)
+		{
+			float Value = ValueData[ChannelIndex];
+
+			for (int32 SampleIndex = ChannelIndex; SampleIndex < NumSamples; SampleIndex += NumChannels)
 			{
-				ProcessAudioNonClamped(FMath::Abs(InAudioBuffer[SampleIndex]));
+				Value -= HistorySquaredData[SampleIndex];
+				Value += InputSquaredData[SampleIndex];
+				OutBuffer[SampleIndex] = Value * NormFactor;
 			}
+
+			ValueData[ChannelIndex] = Value;
 		}
-		return CurrentEnvelopeValue = FMath::Clamp(CurrentEnvelopeValue, 0.0f, 1.0f);
 	}
 
-	float FEnvelopeFollower::ProcessAudio(const float* InAudioBuffer, float* OutAudioBuffer, int32 InNumSamples)
+
+	FMeanSquaredIIR::FMeanSquaredIIR(float InSampleRate, int32 InNumChannels, float InWindowTimeMsec)
+	: SampleRate(InSampleRate)
+	, NumChannels(0)
+	, Alpha(0.f)
+	, Beta(1.f)
 	{
-		// MS/RMS
-		if (EnvMode == EPeakMode::MeanSquared || EnvMode == EPeakMode::RootMeanSquared)
+		if (!ensure(SampleRate > 0.f))
 		{
-			TCircularAudioBuffer<float> SumBuffer(MeanWindowSize);
-			int32 SampleIndex = 0;
-			// This loop iterates per-window, not per-sample
-			while (SampleIndex < InNumSamples)
-			{
-				// Get a new Window
-				const int32 SamplesToCopy = MeanHopSize < (InNumSamples - SampleIndex) ? MeanHopSize : (InNumSamples - SampleIndex);
-				SumBuffer.Push(InAudioBuffer + SampleIndex, SamplesToCopy);
-				float* CurrentWindow = ScratchBuffer.GetData();
-				SumBuffer.Pop(CurrentWindow, MeanWindowSize);
-
-				// Calculate MS/RMS
-				float CurrentMean;
-				ArrayMeanSquared(TArrayView<float>(CurrentWindow, SamplesToCopy), CurrentMean);
-				if (EnvMode == EPeakMode::RootMeanSquared)
-				{
-					CurrentMean = FMath::Sqrt(CurrentMean);
-				}
+			SampleRate = 48000.f;
+		}
 
-				// Use MS/RMS to drive envelope following
-				for (int i = 0; i < SamplesToCopy; ++i)
-				{
-					OutAudioBuffer[SampleIndex + i] = ProcessAudioNonClamped(CurrentMean);
-				}
-				SampleIndex += SamplesToCopy;
+		SetNumChannels(InNumChannels);
+		SetWindowSize(InWindowTimeMsec);
+	}
+
+	void FMeanSquaredIIR::SetWindowSize(float InWindowTimeMsec)
+	{
+		if (ensure(InWindowTimeMsec > 0.f))
+		{
+			// Exponential decay set so window is 1/e at end of window
+			Alpha = FMath::Exp(-1000.f / (SampleRate * InWindowTimeMsec));
+			Beta = 1.f - Alpha;
+		}
+
+		Reset();
+	}
+
+	void FMeanSquaredIIR::SetNumChannels(int32 InNumChannels)
+	{
+		if (NumChannels != InNumChannels)
+		{
+			ChannelValues.Reset();
+			if (ensure(InNumChannels > 0))
+			{
+				ChannelValues.AddZeroed(InNumChannels);
 			}
+			NumChannels = InNumChannels;
 		}
-		// Peak
-		else
+	}
+
+	void FMeanSquaredIIR::Reset()
+	{
+		if (NumChannels > 0)
+		{
+			ChannelValues.Reset();
+			ChannelValues.AddZeroed(NumChannels);
+		}
+	}
+
+	// TODO: what's the convention. pass num frames or num samples?
+	void FMeanSquaredIIR::ProcessAudio(const float* InBuffer, int32 InNumFrames, float* OutBuffer)
+	{
+		const int32 NumSamples = InNumFrames * NumChannels;
+
+		// Square input and store in output.
+		ArraySquare(MakeArrayView(InBuffer, NumSamples), MakeArrayView(OutBuffer, NumSamples));
+
+		// MS[n] = Beta * x^2[n] + Alpha * MS[n - 1];
+		for (int32 ChannelIndex = 0; ChannelIndex < NumChannels; ChannelIndex++)
 		{
-			for (int32 SampleIndex = 0; SampleIndex < InNumSamples; ++SampleIndex)
+			// Get last output value
+			float Value = ChannelValues[ChannelIndex];
+
+			for (int32 SampleIndex = ChannelIndex; SampleIndex < NumSamples; SampleIndex += NumChannels)
 			{
-				OutAudioBuffer[SampleIndex] = ProcessAudioNonClamped(FMath::Abs(InAudioBuffer[SampleIndex]));
+				Value = Beta * OutBuffer[SampleIndex] + Alpha * Value;
+				Value = Audio::UnderflowClamp(Value);
+				OutBuffer[SampleIndex] = Value;
 			}
+
+			// Store last output value for next call to ProcessAudio
+			ChannelValues[ChannelIndex] = Value;
+		}
+	}
+
+
+	// A simple utility that returns a smoothed value given audio input using an RC circuit.
+	// Used for following the envelope of an audio stream.
+	FEnvelopeFollower::FEnvelopeFollower()
+	: FEnvelopeFollower(FEnvelopeFollowerInitParams{})
+	{
+	}
+
+	FEnvelopeFollower::FEnvelopeFollower(const FEnvelopeFollowerInitParams& InParams)
+	: MeanSquaredProcessor(InParams.SampleRate, InParams.NumChannels, InParams.AnalysisWindowMsec)
+	, Smoother(InParams.SampleRate, InParams.NumChannels, InParams.AttackTimeMsec, InParams.ReleaseTimeMsec, InParams.bIsAnalog)
+	, NumChannels(InParams.NumChannels)
+	, EnvMode(InParams.Mode)
+	{
+	}
+
+	// Initialize the envelope follower
+	void FEnvelopeFollower::Init(const FEnvelopeFollowerInitParams& InParams)
+	{
+		Smoother = FAttackReleaseSmoother(InParams.SampleRate, InParams.NumChannels, InParams.AttackTimeMsec, InParams.ReleaseTimeMsec, InParams.bIsAnalog);
+		MeanSquaredProcessor = FMeanSquaredIIR(InParams.SampleRate, InParams.NumChannels, InParams.AnalysisWindowMsec);
+		NumChannels = InParams.NumChannels;
+		EnvMode = InParams.Mode;
+	}
+
+	int32 FEnvelopeFollower::GetNumChannels() const
+	{
+		return NumChannels;
+	}
+
+	void FEnvelopeFollower::SetNumChannels(int32 InNumChannels)
+	{
+		if (InNumChannels != NumChannels)
+		{
+			Smoother.SetNumChannels(InNumChannels);
+			MeanSquaredProcessor.SetNumChannels(NumChannels);
+			NumChannels = InNumChannels;
 		}
+	}
 
-		Audio::BufferRangeClampFast(OutAudioBuffer, InNumSamples, 0.0f, 1.0f);
-		return CurrentEnvelopeValue;
+	// Resets the state of the envelope follower
+	void FEnvelopeFollower::Reset()
+	{
+		MeanSquaredProcessor.Reset();
+		Smoother.Reset();
 	}
 
-	float FEnvelopeFollower::ProcessAudioNonClamped(const float InAudioSample)
+	// Sets whether or not to use analog or digital time constants
+	void FEnvelopeFollower::SetAnalog(bool bInIsAnalog)
 	{
-		float TimeSamples = (InAudioSample > CurrentEnvelopeValue) ? AttackTimeSamples : ReleaseTimeSamples;
-		float NewEnvelopeValue = TimeSamples * (CurrentEnvelopeValue - InAudioSample) + InAudioSample;
-		NewEnvelopeValue = Audio::UnderflowClamp(NewEnvelopeValue);
+		Smoother.SetAnalog(bInIsAnalog);
+	}
 
-		// Update and return the envelope value
-		return CurrentEnvelopeValue = NewEnvelopeValue;
+	// Sets the envelope follower attack time (how fast the envelope responds to input)
+	void FEnvelopeFollower::SetAttackTime(float InAttackTimeMsec)
+	{
+		Smoother.SetAttackTime(InAttackTimeMsec);
+	}
+
+	// Sets the envelope follower release time (how slow the envelope dampens from input)
+	void FEnvelopeFollower::SetReleaseTime(float InReleaseTimeMsec)
+	{
+		Smoother.SetReleaseTime(InReleaseTimeMsec);
 	}
 
-	int16 FEnvelopeFollower::ProcessAudio(const int16 InAudioSample)
+	// Sets the output mode of the envelope follower
+	void FEnvelopeFollower::SetMode(EPeakMode::Type InMode)
 	{
-		// Convert to float
-		float SampleValueFloat = (float)InAudioSample / 32767.0f;
+		EnvMode = InMode;
+	}
 
-		// Process it
-		float Result = ProcessAudio(SampleValueFloat);
+	// Process the input audio buffer and returns the last envelope value
+	void FEnvelopeFollower::ProcessAudio(const float* InBuffer, int32 InNumFrames, float* OutBuffer)
+	{
+		const int32 NumSamples = InNumFrames * NumChannels;
+		if (NumSamples > 0)
+		{
+			ProcessWorkBuffer(InBuffer, InNumFrames);
+			Smoother.ProcessAudio(WorkBuffer.GetData(), InNumFrames, OutBuffer);
+		}
+	}
 
-		// Convert back to int16
-		return (int16)(Result * 32767.0f);
+	void FEnvelopeFollower::ProcessAudio(const float* InBuffer, int32 InNumFrames)
+	{
+		const int32 NumSamples = InNumFrames * NumChannels;
+		if (NumSamples > 0)
+		{
+			ProcessWorkBuffer(InBuffer, InNumFrames);
+			Smoother.ProcessAudio(WorkBuffer.GetData(), InNumFrames);
+		}
 	}
 
-	float FEnvelopeFollower::GetCurrentValue() const
+	void FEnvelopeFollower::ProcessWorkBuffer(const float* InBuffer, int32 InNumFrames)
 	{
-		return CurrentEnvelopeValue;
+		const int32 NumSamples = InNumFrames * NumChannels;
+		if (NumSamples > 0)
+		{
+			WorkBuffer.Reset();
+			WorkBuffer.AddUninitialized(NumSamples);
+
+			if (EPeakMode::RootMeanSquared == EnvMode)
+			{
+				MeanSquaredProcessor.ProcessAudio(InBuffer, InNumFrames, WorkBuffer.GetData());
+				ArraySqrtInPlace(WorkBuffer);
+			}
+			else if (EPeakMode::MeanSquared == EnvMode)
+			{
+				MeanSquaredProcessor.ProcessAudio(InBuffer, InNumFrames, WorkBuffer.GetData());
+			}
+			else if (EPeakMode::Peak == EnvMode)
+			{
+				Audio::ArrayAbs(MakeArrayView(InBuffer, NumSamples), WorkBuffer);
+			}
+			else
+			{
+				check(false);
+			}
+		}
 	}
 
+	const TArray<float>& FEnvelopeFollower::GetEnvelopeValues() const
+	{
+		return Smoother.GetEnvelopeValues();
+	}
 }
diff --git a/Engine/Source/Runtime/SignalProcessing/Private/FloatArrayMath.cpp b/Engine/Source/Runtime/SignalProcessing/Private/FloatArrayMath.cpp
--- a/Engine/Source/Runtime/SignalProcessing/Private/FloatArrayMath.cpp
+++ b/Engine/Source/Runtime/SignalProcessing/Private/FloatArrayMath.cpp
@@ -290,6 +290,21 @@ namespace Audio
 		OutEuclideanNorm = FMath::Sqrt(OutEuclideanNorm);
 	}
 
+	void ArrayAbs(TArrayView<const float> InBuffer, TArrayView<float> OutBuffer)
+	{
+		const int32 Num = InBuffer.Num();
+		check(OutBuffer.Num() == Num);
+
+		const float* InData = InBuffer.GetData();
+		float* OutData = OutBuffer.GetData();
+
+		for (int32 i = 0; i < Num; i++)
+		{
+			OutData[i] = FMath::Abs(InData[i]);
+		}
+	}
+
+
 	void ArrayAbsInPlace(TArrayView<float> InView)
 	{
 		const int32 Num = InView.Num();
@@ -704,6 +719,49 @@ namespace Audio
 		BufferSubtractFast(InMinuend, InSubtrahend, OutArray);
 	}
 
+	void ArraySquare(TArrayView<const float> InValues, TArrayView<float> OutValues)
+	{
+		check(InValues.Num() == OutValues.Num());
+
+		const uint32 Num = InValues.Num();
+		const uint32 NumToSimd = Num & MathIntrinsics::SimdMask;
+
+		const float* InData = InValues.GetData();
+		float* OutData = OutValues.GetData();
+
+		for (uint32 i = 0; i < NumToSimd; i += 4)
+		{
+			VectorRegister4Float VectorData = VectorLoad(&InData[i]);
+			VectorData = VectorMultiply(VectorData, VectorData);
+			VectorStore(VectorData, &OutData[i]);
+		}
+
+		for (uint32 i = NumToSimd; i < Num; i++)
+		{
+			OutData[i] = InData[i] * InData[i];
+		}
+	}
+
+	void ArraySquareInPlace(TArrayView<float> InValues)
+	{
+		const uint32 Num = InValues.Num();
+		const uint32 NumToSimd = Num & MathIntrinsics::SimdMask;
+
+		float* InData = InValues.GetData();
+
+		for (uint32 i = 0; i < NumToSimd; i += 4)
+		{
+			VectorRegister4Float VectorData = VectorLoad(&InData[i]);
+			VectorData = VectorMultiply(VectorData, VectorData);
+			VectorStore(VectorData, &InData[i]);
+		}
+
+		for (uint32 i = NumToSimd; i < Num; i++)
+		{
+			InData[i] = InData[i] * InData[i];
+		}
+	}
+
 	void ArraySqrtInPlace(TArrayView<float> InValues)
 	{
 		const int32 Num = InValues.Num();
diff --git a/Engine/Source/Runtime/SignalProcessing/Public/DSP/Dsp.h b/Engine/Source/Runtime/SignalProcessing/Public/DSP/Dsp.h
--- a/Engine/Source/Runtime/SignalProcessing/Public/DSP/Dsp.h
+++ b/Engine/Source/Runtime/SignalProcessing/Public/DSP/Dsp.h
@@ -806,6 +806,22 @@ namespace Audio
 			SetCapacity(0);
 		}
 
+		TCircularAudioBuffer(const TCircularAudioBuffer<SampleType, Alignment>& InOther)
+		{
+			*this = InOther;
+		}
+
+		TCircularAudioBuffer& operator=(const TCircularAudioBuffer<SampleType, Alignment>& InOther)
+		{
+			InternalBuffer = InOther.InternalBuffer;
+			Capacity = InOther.Capacity;
+			ReadCounter.Set(InOther.ReadCounter.GetValue());
+			WriteCounter.Set(InOther.WriteCounter.GetValue());
+
+			return *this;
+		}
+
+
 		TCircularAudioBuffer(uint32 InCapacity)
 		{
 			SetCapacity(InCapacity);
@@ -851,6 +867,12 @@ namespace Audio
 			}
 		}
 
+		/** Push an array of values into circular buffer. */
+		int32 Push(TArrayView<SampleType> InBuffer)
+		{
+			return Push(InBuffer.GetData(), InBuffer.Num());
+		}
+
 		// Pushes some amount of samples into this circular buffer.
 		// Returns the amount of samples written.
 		// This can only be used for trivially copyable types.
diff --git a/Engine/Source/Runtime/SignalProcessing/Public/DSP/DynamicsProcessor.h b/Engine/Source/Runtime/SignalProcessing/Public/DSP/DynamicsProcessor.h
--- a/Engine/Source/Runtime/SignalProcessing/Public/DSP/DynamicsProcessor.h
+++ b/Engine/Source/Runtime/SignalProcessing/Public/DSP/DynamicsProcessor.h
@@ -95,7 +95,7 @@ namespace Audio
 		TArray<FDelay> LookaheadDelay;
 
 		// Envelope followers
-		TArray<FEnvelopeFollower> EnvFollower;
+		TArray<FInlineEnvelopeFollower> EnvFollower;
 
 		// Points in the knee used for lagrangian interpolation
 		TArray<FVector2D> KneePoints;
diff --git a/Engine/Source/Runtime/SignalProcessing/Public/DSP/EnvelopeFollower.h b/Engine/Source/Runtime/SignalProcessing/Public/DSP/EnvelopeFollower.h
--- a/Engine/Source/Runtime/SignalProcessing/Public/DSP/EnvelopeFollower.h
+++ b/Engine/Source/Runtime/SignalProcessing/Public/DSP/EnvelopeFollower.h
@@ -4,6 +4,7 @@
 
 #include "CoreMinimal.h"
 #include "DSP/Dsp.h"
+#include "DSP/FloatArrayMath.h"
 
 namespace Audio
 {
@@ -19,82 +20,405 @@ namespace Audio
 		};
 	}
 
-	// A simple utility that returns a smoothed value given audio input using an RC circuit.
-	// Used for following the envelope of an audio stream.
-	class SIGNALPROCESSING_API FEnvelopeFollower
+
+	/** Conversion between attack/release time and attack/release sample counts. */
+	class SIGNALPROCESSING_API FAttackRelease
 	{
+		// see https://en.wikipedia.org/wiki/RC_time_constant
+		// Time constants indicate how quickly the envelope follower responds to changes in input
+		static constexpr float AnalogTimeConstant = 1.00239343f;
+		static constexpr float DigitalTimeConstant = 4.60517019f;
 	public:
 
-		static constexpr int32 DefaultWindowSize = 100;
-		static constexpr int32 DefaultHopSize = DefaultWindowSize / 2;
-		static_assert(DefaultHopSize > 0);
+		/** Construct an FAttackRelease object.
+		 *
+		 * @param InSampleRate - The number of frames per a second.
+		 * @param InAttackTimeMsec - The desired attack time in milliseconds.
+		 * @param InReleaseTimeMsec - The desired release time in milliseconds.
+		 * @param bInIsAnalog - Whether to model analog RC circuits or use digital models.
+		 */
+		FAttackRelease(float InSampleRate, float InAttackTimeMsec, float InReleaseTimeMsec, bool bInIsAnalog);
+
+		void SetAnalog(bool bInIsAnalog);
+		void SetAttackTime(float InAttackTimeMsec);
+		void SetReleaseTime(float InReleaseTimeMsec);
+
+		/** Get the attack time in samples. */
+		FORCEINLINE float GetAttackTimeSamples() const { return AttackTimeSamples; }
+
+		/** Get the release time in samples. */
+		FORCEINLINE float GetReleaseTimeSamples() const { return ReleaseTimeSamples; }
+
+		/** Get the sample rate. */
+		FORCEINLINE float GetSampleRate() const { return SampleRate; }
+
+	protected:
+
+		void SetSampleRate(float InSampleRate);
+
+	private:
+
+		float SampleRate;
+		float AttackTimeSamples;
+		float AttackTimeMsec;
+		float ReleaseTimeSamples;
+		float ReleaseTimeMsec;
+		bool bIsAnalog;
+	};
+	
+
+	/** Smooths signals using attack and release settings. */
+	class SIGNALPROCESSING_API FAttackReleaseSmoother : public FAttackRelease
+	{
+	public:
+		/** Construct an FAttackReleaseSmoother.
+		 *
+		 * @param InSampleRate - The number of samples per a second of incoming data.
+		 * @param InNumChannels - The number of input channels. 
+		 * @param InAttackTimeMsec - The desired attack time in milliseconds.
+		 * @param InReleaseTimeMsec - The desired release time in milliseconds.
+		 * @param bInIsAnalog - Whether to model analog RC circuits or use digital models.
+		 */
+		FAttackReleaseSmoother(float InSampleRate, int32 InNumChannels, float InAttackTimeMsec, float InReleaseTimeMsec, bool bInIsAnalog);
+
+		/** Process channel interleaved data.
+		 *
+		 * @param InBuffer - Interleaved channel data. 
+		 * @param InNumFrames - Number of frames in InBuffer.
+		 */
+		void ProcessAudio(const float* InBuffer, int32 InNumFrames);
+
+		/** Process channel interleaved data.
+		 *
+		 * @param InBuffer - Interleaved input data. 
+		 * @param InNumFrames - Number of frames in InBuffer.
+		 * @param OutBuffer - Output interleaved envelope data. Should be same size as InBuffer.
+		 */
+		void ProcessAudio(const float* InBuffer, int32 InNumFrames, float* OutBuffer);
+
+		/** Retrieve the final values of the envelope for each channel. */
+		const TArray<float>& GetEnvelopeValues() const;
+
+		/** Set the number of input channels */
+		void SetNumChannels(int32 InNumChannels);
 
-		FEnvelopeFollower();
-		FEnvelopeFollower(const float InSampleRate,
-			const float InAttackTimeMsec = 10.0f,
-			const float InReleaseTimeMSec = 100.0f,
-			const EPeakMode::Type InMode = EPeakMode::RootMeanSquared,
-			const bool bInIsAnalg = true,
-			const int32 InWindowSizeForMean = DefaultWindowSize,
-			const int32 InHopSizeForMean = DefaultHopSize);
-
-		virtual ~FEnvelopeFollower();
-
-		// Initialize the envelope follower
-		void Init(const float InSampleRate,
-			const float InAttackTimeMsec = 10.0f,
-			const float InReleaseTimeMSec = 100.0f,
-			const EPeakMode::Type InMode = EPeakMode::RootMeanSquared,
-			const bool bInIsAnalg = true,
-			const int32 InWindowSizeForMean = DefaultWindowSize,
-			const int32 InHopSizeForMean = DefaultHopSize);
-
-		// Resets the state of the envelope follower
 		void Reset();
 
-		// Sets whether or not to use analog or digital time constants
-		void SetAnalog(const bool bInIsAnalog);
+	private:
+		int32 NumChannels;
+		TArray<float> PriorEnvelopeValues;
+	};
+
+	/** Compute mean squared using FIR method.  */
+	class SIGNALPROCESSING_API FMeanSquaredFIR
+	{
+		static constexpr int32 DefaultHistoryCapacity = 16384;
+
+	public:
 
-		// Sets the envelope follower attack time (how fast the envelope responds to input)
-		void SetAttackTime(const float InAttackTimeMsec);
+		/** Construct an FMeanSquaredFIR
+		 *
+		 * @parma InSampleRate - Number of frames per a second.
+		 * @param InNumChannels - Number of channels per a frame.
+		 * @param InWindowTimeMsec - Number of milliseconds per a mean squared window.
+		 */
+		FMeanSquaredFIR(float InSampleRate, int32 InNumChannels, float InWindowTimeMsec);
 
-		// Sets the envelope follower release time (how slow the envelope dampens from input)
-		void SetReleaseTime(const float InReleaseTimeMsec);
+		/** Set the size of the analysis window. */
+		void SetWindowSize(float InWindowTimeMsec);
+
+		/** Set the number of input channels. */
+		void SetNumChannels(int32 InNumChannels);
+
+		void Reset();
 
-		// Sets the output mode of the envelope follower
-		void SetMode(const EPeakMode::Type InMode);
+		/** Calculate mean squared per sample.
+		 *
+		 * @param InBuffer - Interleaved input data. 
+		 * @param InNumFrames - Number of frames in InBuffer.
+		 * @param OutBuffer - Output interleaved data. Should be same size as InBuffer.
+		 */
+		void ProcessAudio(const float* InBuffer, int32 InNumFrames, float* OutBuffer);
 
-		// Processes the input audio stream and returns the envelope value.
-		float ProcessAudio(const float InAudioSample);
+	private:
+		float SampleRate;
+		int32 NumChannels;
+		int32 WindowTimeFrames;
+		int32 WindowTimeSamples;
+		float NormFactor;
 
-		// Process the input audio buffer and returns the last envelope value
-		float ProcessAudio(const float* InAudioBuffer, float* OutAudioBuffer, int32 InNumSamples);
+		TArray<float> ChannelValues;
+		TCircularAudioBuffer<float> HistorySquared;
+		FAlignedFloatBuffer SquaredHistoryBuffer;
+		FAlignedFloatBuffer SquaredInputBuffer;
+	};
 
-		// Process the input audio buffer and returns the last envelope value
-		float ProcessAudio(const float* InAudioBuffer, int32 InNumSamples);
+	/** Compute mean squared using IIR method.  */
+	class SIGNALPROCESSING_API FMeanSquaredIIR
+	{
+	public:
+		/** Construct an FMeanSquaredIIR
+		 *
+		 * @parma InSampleRate - Number of frames per a second.
+		 * @param InNumChannels - Number of channels per a frame.
+		 * @param InWindowTimeMsec - Number of milliseconds per a mean squared window.
+		 */
+		FMeanSquaredIIR(float InSampleRate, int32 InNumChannels, float InWindowTimeMsec);
 
-		// Processes the input audio stream and returns the envelope value.
-		float ProcessAudioNonClamped(const float InAudioSample);
+		/** Set the size of the analysis window. */
+		void SetWindowSize(float InWindowTimeMsec);
 
-		// Process the input audio stream in int16 format
-		int16 ProcessAudio(const int16 InAudioSample);
+		/** Set the number of input channels. */
+		void SetNumChannels(int32 InNumChannels);
 
-		// Gets the current envelope value
-		float GetCurrentValue() const;
+		void Reset();
 
-	protected:
-		EPeakMode::Type EnvMode;
-		int32 MeanWindowSize;
-		int32 MeanHopSize;		
-		TArray<float> ScratchBuffer;
+		/** Calculate mean squared per sample.
+		 *
+		 * @param InBuffer - Interleaved input data. 
+		 * @param InNumFrames - Number of frames in InBuffer.
+		 * @param OutBuffer - Output interleaved data. Should be same size as InBuffer.
+		 */
+		void ProcessAudio(const float* InBuffer, int32 InNumFrames, float* OutBuffer);
+
+	private:
 		float SampleRate;
-		float AttackTimeMsec;
-		float AttackTimeSamples;
-		float ReleaseTimeMsec;
-		float ReleaseTimeSamples;
-		float CurrentEnvelopeValue;
-		bool bIsAnalog;
+		int32 NumChannels;
+		float Alpha;
+		float Beta;
+
+		TArray<float> ChannelValues;
 	};
 
+	struct FEnvelopeFollowerInitParams
+	{
+		/** Number of frames per a second. */
+		float SampleRate = 48000.f;
+		/** Number of channels per a frame. */
+		int32 NumChannels = 1;
+		/** The desired attack time in milliseconds. */
+		float AttackTimeMsec = 10.0f;
+		/** The desired release time in milliseconds. */
+		float ReleaseTimeMsec = 100.0f;
+		/** Technique for measuring amplitude of input signal. */
+		EPeakMode::Type Mode = EPeakMode::Peak;
+		/** Whether to model analog RC circuits or use digital models. */
+		bool bIsAnalog = true;
+		/** Number of milliseconds per a (root) mean squared window. */
+		float AnalysisWindowMsec=5.f;
+	};
 
+	/** A simple utility that returns a smoothed value given audio input using 
+	 * an RC circuit.  Used for following the envelope of an audio stream.
+	 */
+	class SIGNALPROCESSING_API FEnvelopeFollower
+	{
+	public:
+		/** Construct an envelope follower. */
+		FEnvelopeFollower();
+
+		/** Construct an envelope follower. */
+		FEnvelopeFollower(const FEnvelopeFollowerInitParams& InParams);
+
+		/** Initialize the envelope follower. */
+		void Init(const FEnvelopeFollowerInitParams& InParams);
+
+		int32 GetNumChannels() const;
+
+		/** Set the number of channels per an input frame. */
+		void SetNumChannels(int32 InNumChannels);
+
+		/** Sets whether or not to use analog or digital time constants */
+		void SetAnalog(bool bInIsAnalog);
+
+		/** Sets the envelope follower attack time (how fast the envelope responds to input) */
+		void SetAttackTime(float InAttackTimeMsec);
+
+		/** Sets the envelope follower release time (how slow the envelope dampens from input) */
+		void SetReleaseTime(float InReleaseTimeMsec);
+
+		/** Sets the input mode of the envelope follower */
+		void SetMode(EPeakMode::Type InMode);
+
+		/** Calculate envelope per sample.
+		 *
+		 * @param InBuffer - Interleaved input data. 
+		 * @param InNumFrames - Number of frames in InBuffer.
+		 * @param OutBuffer - Output interleaved data. Should be same size as InBuffer.
+		 */
+		void ProcessAudio(const float* InBuffer, int32 InNumFrames, float* OutBuffer);
+
+		/** Calculate envelope
+		 *
+		 * @param InBuffer - Interleaved input data. 
+		 * @param InNumFrames - Number of frames in InBuffer.
+		 */
+		void ProcessAudio(const float* InBuffer, int32 InNumFrames);
+		
+		/** Retrieve the final values of the envelope for each channel. */
+		const TArray<float>& GetEnvelopeValues() const;
+
+		/** Resets the state of the envelope follower */
+		void Reset();
+
+	private:
+		void ProcessWorkBuffer(const float* InBuffer, int32 InNumFrames);
+
+		FAlignedFloatBuffer WorkBuffer;
+		FMeanSquaredIIR MeanSquaredProcessor; 
+		FAttackReleaseSmoother Smoother;
+
+		int32 NumChannels;
+		EPeakMode::Type EnvMode;
+	};
+
+	struct FInlineEnvelopeFollowerInitParams
+	{
+		/** Number of frames per a second. */
+		float SampleRate = 48000.f;
+		/** The desired attack time in milliseconds. */
+		float AttackTimeMsec = 10.0f;
+		/** The desired release time in milliseconds. */
+		float ReleaseTimeMsec = 100.0f;
+		/** Technique for measuring amplitude of input signal. */
+		EPeakMode::Type Mode = EPeakMode::Peak;
+		/** Whether to model analog RC circuits or use digital models. */
+		bool bIsAnalog = true;
+		/** Number of milliseconds per a (root) mean squared window. */
+		float AnalysisWindowMsec=5.f;
+	};
+
+	/** FInlineEnvlepeFollower is useful for low samplerate use cases and where
+	 * samples are only available one at a time. This class is inlined becaused
+	 * there are situations where it is needed in a CPU intensive situations.
+	 */
+	class FInlineEnvelopeFollower : public FAttackRelease
+	{
+
+	public:
+
+		/** Construct an envelope follower. */
+		FInlineEnvelopeFollower(const FInlineEnvelopeFollowerInitParams& InParams)
+		: FAttackRelease(InParams.SampleRate, InParams.AttackTimeMsec, InParams.ReleaseTimeMsec, InParams.bIsAnalog)
+		, Value(0.f)
+		, Mode(InParams.Mode)
+		, AnalysisValue(0.f)
+		, AnalysisWindowMsec(5.f)
+		{
+			SetAnalysisWindow(InParams.AnalysisWindowMsec);
+		}
+
+		FInlineEnvelopeFollower()
+		: FInlineEnvelopeFollower(FInlineEnvelopeFollowerInitParams{})
+		{
+		}
+
+		/** Initialize an envelope follower.  */
+		void Init(const FInlineEnvelopeFollowerInitParams& InParams)
+		{
+			SetSampleRate(InParams.SampleRate);
+			SetAttackTime(InParams.AttackTimeMsec);
+			SetReleaseTime(InParams.ReleaseTimeMsec);
+			SetMode(InParams.Mode);
+			SetAnalog(InParams.bIsAnalog);
+			SetAnalysisWindow(InParams.AnalysisWindowMsec);
+		}
+
+		/* Sets the input analysis mode of the envelope follower */
+		void SetMode(EPeakMode::Type InMode)
+		{
+			Mode = InMode;
+		}
+
+		/** Set the analysis window size (for MeanSquared and RootMeanSquared). */
+		void SetAnalysisWindow(float InAnalysisWindowMsec)
+		{
+			if (ensure(InAnalysisWindowMsec > 0.f))
+			{
+				AnalysisWindowMsec = InAnalysisWindowMsec;
+				InitMeanSquaredCoefficients(AnalysisWindowMsec, GetSampleRate());
+			}
+		}
+
+		/** Process a single sample and return the envelope value. */
+		FORCEINLINE float ProcessSample(float InSample)
+		{
+			float NormSample = NormalizeSample(InSample);
+
+			float Diff = Value - NormSample;
+			if (Diff <= 0.f)
+			{
+				Value = (GetAttackTimeSamples() * Diff) + NormSample;
+			}
+			else
+			{
+				Value = (GetReleaseTimeSamples() * Diff) + NormSample;
+			}
+
+			Value = Audio::UnderflowClamp(Value);
+			return Value;
+		}
+
+		void Reset()
+		{
+			Value = 0.f;
+		}
+
+		/** Return the most recent envelope value. */
+		FORCEINLINE float GetValue() const
+		{
+			return Value;
+		}
+
+	private:
+		FORCEINLINE float NormalizeSample(float InSample)
+		{
+			switch (Mode)
+			{
+				case EPeakMode::Peak:
+					return FMath::Abs(InSample);
+					
+				case EPeakMode::MeanSquared:
+
+					{
+						float SampleSquared = InSample * InSample;
+						AnalysisValue = AnalysisFilterBeta * SampleSquared + AnalysisFilterAlpha * AnalysisValue;
+						AnalysisValue = Audio::UnderflowClamp(AnalysisValue);
+						return AnalysisValue;
+					}
+
+				case EPeakMode::RootMeanSquared:
+
+					{
+						float SampleSquared = InSample * InSample;
+						AnalysisValue = AnalysisFilterBeta * SampleSquared + AnalysisFilterAlpha * AnalysisValue;
+						AnalysisValue = Audio::UnderflowClamp(AnalysisValue);
+						return FMath::Sqrt(AnalysisValue);
+					}
+
+				default:
+				{
+					checkNoEntry();
+				}
+			}
+
+			return InSample;
+		}
+
+		void InitMeanSquaredCoefficients(float InWinMsec, float InSampleRate)
+		{
+			if (ensure((InSampleRate > 0.f) && (InWinMsec > 0.f)))
+			{
+				AnalysisFilterAlpha = FMath::Exp(-1000.f / (InSampleRate * InWinMsec));
+				AnalysisFilterBeta = 1.f - AnalysisFilterAlpha;
+			}
+		}
+
+		float Value = 0.f;
+		EPeakMode::Type Mode = EPeakMode::Peak;
+		float AnalysisValue = 0.f;
+		float AnalysisWindowMsec = 10.f;
+		float AnalysisFilterAlpha = 0.f;
+		float AnalysisFilterBeta = 1.f;
+	};
 }
diff --git a/Engine/Source/Runtime/SignalProcessing/Public/DSP/FloatArrayMath.h b/Engine/Source/Runtime/SignalProcessing/Public/DSP/FloatArrayMath.h
--- a/Engine/Source/Runtime/SignalProcessing/Public/DSP/FloatArrayMath.h
+++ b/Engine/Source/Runtime/SignalProcessing/Public/DSP/FloatArrayMath.h
@@ -7,6 +7,7 @@
 
 namespace Audio
 {
+	/** Sum all values in an array. */
 	SIGNALPROCESSING_API void ArraySum(TArrayView<const float> InValues, float& OutSum);
 	SIGNALPROCESSING_API void ArraySum(const FAlignedFloatBuffer& InValues, float& OutSum);
 
@@ -57,11 +58,14 @@ namespace Audio
 	/** Computes the EuclideanNorm of the InView. Same as calculating the energy in window. */
 	SIGNALPROCESSING_API void ArrayGetEuclideanNorm(TArrayView<const float> InView, float& OutEuclideanNorm);
 
+	/** Absolute value of array elements */
+	SIGNALPROCESSING_API void ArrayAbs(TArrayView<const float> InBuffer, TArrayView<float> OutBuffer);
+
 	/** Absolute value of array elements in place.
 	 *
-	 *  InView contains the data to be manipulated.
+	 *  InBuffer contains the data to be manipulated.
 	 */
-	SIGNALPROCESSING_API void ArrayAbsInPlace(TArrayView<float> InView);
+	SIGNALPROCESSING_API void ArrayAbsInPlace(TArrayView<float> InBuffer);
 
 	/** Clamp minimum value of array in place.
 	 *
@@ -148,6 +152,11 @@ namespace Audio
 	SIGNALPROCESSING_API void ArraySubtractByConstantInPlace(TArrayView<float> InValues, float InSubtrahend);
 	SIGNALPROCESSING_API void ArraySubtractByConstantInPlace(FAlignedFloatBuffer& InValues, float InSubtrahend);
 
+	/* Square values */
+	SIGNALPROCESSING_API void ArraySquare(TArrayView<const float> InValues, TArrayView<float> OutValues);
+
+	/** Square values in place. */
+	SIGNALPROCESSING_API void ArraySquareInPlace(TArrayView<float> InValues);
 
 	/** Take Square Root of values in place. */
 	SIGNALPROCESSING_API void ArraySqrtInPlace(TArrayView<float> InValues);
