Fix for submix envelope follower.

diff --git a/Engine/Source/Runtime/SignalProcessing/Private/EnvelopeFollower.cpp b/Engine/Source/Runtime/SignalProcessing/Private/EnvelopeFollower.cpp
--- a/Engine/Source/Runtime/SignalProcessing/Private/EnvelopeFollower.cpp
+++ b/Engine/Source/Runtime/SignalProcessing/Private/EnvelopeFollower.cpp
@@ -5,6 +5,26 @@
 
 namespace Audio
 {
+	namespace EnvelopeFollowerPrivate
+	{
+		FORCEINLINE float SmoothSample(float InSample, float InPriorSmoothedSample, float InAttackSamples, float InReleaseSamples)
+		{
+			float Value = InPriorSmoothedSample;
+			float Diff = Value - InSample;
+
+			if (Diff <= 0.f)
+			{
+				Value = (InAttackSamples * Diff) + InSample;
+			}
+			else
+			{
+				Value = (InReleaseSamples * Diff) + InSample;
+			}
+
+			return Audio::UnderflowClamp(Value);
+		}
+	}
+
 	FAttackRelease::FAttackRelease(float InSampleRate, float InAttackTimeMsec, float InReleaseTimeMsec, bool bInIsAnalog)
 	: SampleRate(InSampleRate)
 	, AttackTimeSamples(1.f)
@@ -86,17 +106,7 @@ namespace Audio
 
 				for (int32 Pos = Channel; Pos < NumSamples; Pos += NumChannels)
 				{
-					float Diff = EnvelopeValue - InBuffer[Pos];
-					if (Diff > 0.f)
-					{
-						EnvelopeValue = (AttackSamples * Diff) + InBuffer[Pos];
-					}
-					else
-					{
-						EnvelopeValue = (ReleaseSamples * Diff) + InBuffer[Pos];
-					}
-
-					EnvelopeValue = Audio::UnderflowClamp(EnvelopeValue);
+					EnvelopeValue = EnvelopeFollowerPrivate::SmoothSample(InBuffer[Pos], EnvelopeValue, AttackSamples, ReleaseSamples);
 				}
 
 				PriorEnvelopeValues[Channel] = EnvelopeValue;
@@ -118,17 +128,7 @@ namespace Audio
 
 				for (int32 Pos = Channel; Pos < NumSamples; Pos += NumChannels)
 				{
-					float Diff = EnvelopeValue - InBuffer[Pos];
-					if (Diff <= 0.f)
-					{
-						EnvelopeValue = (AttackSamples * Diff) + InBuffer[Pos];
-					}
-					else
-					{
-						EnvelopeValue = (ReleaseSamples * Diff) + InBuffer[Pos];
-					}
-
-					EnvelopeValue = Audio::UnderflowClamp(EnvelopeValue);
+					EnvelopeValue = EnvelopeFollowerPrivate::SmoothSample(InBuffer[Pos], EnvelopeValue, AttackSamples, ReleaseSamples);
 					OutBuffer[Pos] = EnvelopeValue;
 				}
 
