Add getters for smoother members to EnvelopeFollower
#rnx

diff --git a/Engine/Source/Runtime/SignalProcessing/Private/EnvelopeFollower.cpp b/Engine/Source/Runtime/SignalProcessing/Private/EnvelopeFollower.cpp
--- a/Engine/Source/Runtime/SignalProcessing/Private/EnvelopeFollower.cpp
+++ b/Engine/Source/Runtime/SignalProcessing/Private/EnvelopeFollower.cpp
@@ -366,6 +366,31 @@ namespace Audio
 		return NumChannels;
 	}
 
+	float FEnvelopeFollower::GetSampleRate() const
+	{
+		return Smoother.GetSampleRate();
+	}
+
+	float FEnvelopeFollower::GetAttackTime() const
+	{
+		return Smoother.GetAttackTimeSamples() / Smoother.GetSampleRate();
+	}
+
+	float FEnvelopeFollower::GetReleaseTime() const
+	{
+		return Smoother.GetReleaseTimeSamples() / Smoother.GetSampleRate();
+	}
+
+	bool FEnvelopeFollower::GetAnalog() const
+	{
+		return Smoother.GetAnalog();
+	}
+
+	EPeakMode::Type FEnvelopeFollower::GetMode() const
+	{
+		return EnvMode;
+	}
+
 	void FEnvelopeFollower::SetNumChannels(int32 InNumChannels)
 	{
 		if (InNumChannels != NumChannels)
diff --git a/Engine/Source/Runtime/SignalProcessing/Public/DSP/EnvelopeFollower.h b/Engine/Source/Runtime/SignalProcessing/Public/DSP/EnvelopeFollower.h
--- a/Engine/Source/Runtime/SignalProcessing/Public/DSP/EnvelopeFollower.h
+++ b/Engine/Source/Runtime/SignalProcessing/Public/DSP/EnvelopeFollower.h
@@ -28,8 +28,8 @@ namespace Audio
 		// Time constants indicate how quickly the envelope follower responds to changes in input
 		static constexpr float AnalogTimeConstant = 1.00239343f;
 		static constexpr float DigitalTimeConstant = 4.60517019f;
-	public:
 
+	public:
 		/** Construct an FAttackRelease object.
 		 *
 		 * @param InSampleRate - The number of frames per a second.
@@ -43,6 +43,9 @@ namespace Audio
 		void SetAttackTime(float InAttackTimeMsec);
 		void SetReleaseTime(float InReleaseTimeMsec);
 
+		/** Get whether set to analog or digital time constant. (True is analog, false is digital) */
+		FORCEINLINE bool GetAnalog() const { return bIsAnalog; }
+
 		/** Get the attack time in samples. */
 		FORCEINLINE float GetAttackTimeSamples() const { return AttackTimeSamples; }
 
@@ -223,8 +226,24 @@ namespace Audio
 		/** Initialize the envelope follower. */
 		void Init(const FEnvelopeFollowerInitParams& InParams);
 
+		/** Returns the number of channels per an input frame */
 		int32 GetNumChannels() const;
 
+		/** Returns the number of frames per a second set on initialization */
+		float GetSampleRate() const;
+
+		/** Returns the envelope follower attack time (how fast the envelope responds to input) */
+		float GetAttackTime() const;
+
+		/** Returns the envelope follower release time (how slow the envelope dampens from input) */
+		float GetReleaseTime() const;
+
+		/** Returns whether or not to use analog or digital time constants */
+		bool GetAnalog() const;
+
+		/** Returns the input mode of the envelope follower */
+		EPeakMode::Type GetMode() const;
+
 		/** Set the number of channels per an input frame. */
 		void SetNumChannels(int32 InNumChannels);
 
