Fixed memory stomp in FFT filter

diff --git a/Engine/Source/Runtime/SignalProcessing/Public/DSP/PassiveFilter.h b/Engine/Source/Runtime/SignalProcessing/Public/DSP/PassiveFilter.h
--- a/Engine/Source/Runtime/SignalProcessing/Public/DSP/PassiveFilter.h
+++ b/Engine/Source/Runtime/SignalProcessing/Public/DSP/PassiveFilter.h
@@ -191,8 +191,8 @@ namespace Audio
 			FreqData.OutReal[Index] *= Gain;
 			FreqData.OutImag[Index] *= Gain;
 
-			FreqData.OutReal[NumSamples - Index] *= Gain;
-			FreqData.OutImag[NumSamples - Index] *= Gain;
+			FreqData.OutReal[NumSamples - Index - 1] *= Gain;
+			FreqData.OutImag[NumSamples - Index - 1] *= Gain;
 		});
 
 		// Inverse FFT back into the signal:
