Fixing and adding more logging to voice chat attenuation.
Refactored the way data is retrieved in voice chat manager.

diff --git a/Engine/Source/Runtime/SignalProcessing/Private/AudioBufferDistanceAttenuation.cpp b/Engine/Source/Runtime/SignalProcessing/Private/AudioBufferDistanceAttenuation.cpp
--- a/Engine/Source/Runtime/SignalProcessing/Private/AudioBufferDistanceAttenuation.cpp
+++ b/Engine/Source/Runtime/SignalProcessing/Private/AudioBufferDistanceAttenuation.cpp
@@ -43,11 +43,12 @@ namespace Audio
 
 		// TODO: investigate adding int16 flavors of utilities in BufferVectorOperations.h to avoid format conversions
 		uint32 CurrentSampleIndex = 0;
-		const float DeltaValue = (TargetAttenuationLinear - InOutAttenuation) / FrameCount;
 
 		// If we're passed in a negative value for InOutAttenuation, that means we don't want to interpolate from that value to target value (i.e. it's the first one).
 		// This prevents a pop when first applying attenuation if a sound is far away.
 		float Gain = InOutAttenuation < 0.0f ? TargetAttenuationLinear : InOutAttenuation;
+		const float DeltaValue = (TargetAttenuationLinear - Gain) / FrameCount;
+
 		for (int32 FrameIndex = 0; FrameIndex < FrameCount; ++FrameIndex)
 		{
 			for (uint32 ChannelIndex = 0; ChannelIndex < InNumChannels; ++ChannelIndex)
