Fix for biquad filter going silent after going past nyquist.

diff --git a/Engine/Source/Runtime/SignalProcessing/Private/Filter.cpp b/Engine/Source/Runtime/SignalProcessing/Private/Filter.cpp
--- a/Engine/Source/Runtime/SignalProcessing/Private/Filter.cpp
+++ b/Engine/Source/Runtime/SignalProcessing/Private/Filter.cpp
@@ -296,11 +296,27 @@ namespace Audio
 			break;
 		}
 
-		a0 /= b0;
-		a1 /= b0;
-		a2 /= b0;
-		b1 /= b0;
-		b2 /= b0;
+		// Protect against indefinite filter coefficients. Once a NaN or inf gets
+		// into the filter, it cannot return to normal operation unless the history
+		// samples are reset. 
+		if (FMath::IsFinite(b0))
+		{
+			a0 /= b0;
+			a1 /= b0;
+			a2 /= b0;
+			b1 /= b0;
+			b2 /= b0;
+		}
+		else
+		{
+			// Coefficients for silence.
+			a0 = 0.f;
+			a1 = 0.f;
+			a2 = 0.f;
+			b0 = 1.f;
+			b1 = 0.f;
+			b2 = 0.f;
+		}
 
 		for (int32 Channel = 0; Channel < NumChannels; ++Channel)
 		{
