Clean up type truncation warnings

#rnx

diff --git a/Engine/Source/Runtime/Engine/Classes/Engine/Texture2D.h b/Engine/Source/Runtime/Engine/Classes/Engine/Texture2D.h
--- a/Engine/Source/Runtime/Engine/Classes/Engine/Texture2D.h
+++ b/Engine/Source/Runtime/Engine/Classes/Engine/Texture2D.h
@@ -119,9 +119,9 @@ public:
 	//~ End UObject Interface.
 
 	//~ Begin UTexture Interface.
-	virtual float GetSurfaceWidth() const override { return GetSizeX(); }
-	virtual float GetSurfaceHeight() const override { return GetSizeY(); }
-	virtual float GetSurfaceDepth() const override { return 0; }
+	virtual float GetSurfaceWidth() const override { return static_cast<float>(GetSizeX()); }
+	virtual float GetSurfaceHeight() const override { return static_cast<float>(GetSizeY()); }
+	virtual float GetSurfaceDepth() const override { return 0.0f; }
 	virtual uint32 GetSurfaceArraySize() const override { return 0; }
 	virtual TextureAddress GetTextureAddressX() const override { return AddressX; }
 	virtual TextureAddress GetTextureAddressY() const override { return AddressY; }
diff --git a/Engine/Source/Runtime/Engine/Classes/Engine/Texture2DArray.h b/Engine/Source/Runtime/Engine/Classes/Engine/Texture2DArray.h
--- a/Engine/Source/Runtime/Engine/Classes/Engine/Texture2DArray.h
+++ b/Engine/Source/Runtime/Engine/Classes/Engine/Texture2DArray.h
@@ -53,9 +53,9 @@ public:
 	virtual void GetAssetRegistryTags(TArray<FAssetRegistryTag>& OutTags) const override;
 	virtual FString GetDesc() override;
 	virtual void GetResourceSizeEx(FResourceSizeEx& CumulativeResourceSize) override;
-	virtual float GetSurfaceWidth() const override { return GetSizeX(); }
-	virtual float GetSurfaceHeight() const override { return GetSizeY(); }
-	virtual float GetSurfaceDepth() const override { return 0; }
+	virtual float GetSurfaceWidth() const override { return static_cast<float>(GetSizeX()); }
+	virtual float GetSurfaceHeight() const override { return static_cast<float>(GetSizeY()); }
+	virtual float GetSurfaceDepth() const override { return 0.0f; }
 	virtual uint32 GetSurfaceArraySize() const override { return GetArraySize(); }
 	virtual TextureAddress GetTextureAddressX() const override { return AddressX; }
 	virtual TextureAddress GetTextureAddressY() const override { return AddressY; }
diff --git a/Engine/Source/Runtime/Engine/Classes/Engine/TextureCubeArray.h b/Engine/Source/Runtime/Engine/Classes/Engine/TextureCubeArray.h
--- a/Engine/Source/Runtime/Engine/Classes/Engine/TextureCubeArray.h
+++ b/Engine/Source/Runtime/Engine/Classes/Engine/TextureCubeArray.h
@@ -49,9 +49,9 @@ public:
 	virtual void GetAssetRegistryTags(TArray<FAssetRegistryTag>& OutTags) const override;
 	virtual FString GetDesc() override;
 	virtual void GetResourceSizeEx(FResourceSizeEx& CumulativeResourceSize) override;
-	virtual float GetSurfaceWidth() const override { return GetSizeX(); }
-	virtual float GetSurfaceHeight() const override { return GetSizeY(); }
-	virtual float GetSurfaceDepth() const override { return 0; }
+	virtual float GetSurfaceWidth() const override { return static_cast<float>(GetSizeX()); }
+	virtual float GetSurfaceHeight() const override { return static_cast<float>(GetSizeY()); }
+	virtual float GetSurfaceDepth() const override { return 0.0f; }
 	virtual uint32 GetSurfaceArraySize() const override { return GetNumSlices(); }
 	virtual FTextureResource* CreateResource() override;
 #if WITH_EDITOR
diff --git a/Engine/Source/Runtime/SignalProcessing/Public/DSP/Dsp.h b/Engine/Source/Runtime/SignalProcessing/Public/DSP/Dsp.h
--- a/Engine/Source/Runtime/SignalProcessing/Public/DSP/Dsp.h
+++ b/Engine/Source/Runtime/SignalProcessing/Public/DSP/Dsp.h
@@ -305,18 +305,18 @@ namespace Audio
 		// Check if equal as well as less than to avoid round error in case where at edges.
 		if (InValue <= Domain.X)
 		{
-			return Range.X;
+			return UE_REAL_TO_FLOAT(Range.X);
 		}
 
 		if (InValue >= Domain.Y)
 		{
-			return Range.Y;
+			return UE_REAL_TO_FLOAT(Range.Y);
 		}
 
 		//Handle edge case of NaN
 		if (FMath::IsNaN(InValue))
 		{
-			return Range.X;
+			return UE_REAL_TO_FLOAT(Range.X);
 		}
 
 		const FVector2D RangeLog(FMath::Max(FMath::Loge(Range.X), SMALL_NUMBER), FMath::Min(FMath::Loge(Range.Y), BIG_NUMBER));
@@ -330,23 +330,23 @@ namespace Audio
 		// Check if equal as well as less than to avoid round error in case where at edges.
 		if (InFrequencyValue <= Range.X)
 		{
-			return Domain.X;
+			return UE_REAL_TO_FLOAT(Domain.X);
 		}
 
 		if (InFrequencyValue >= Range.Y)
 		{
-			return Domain.Y;
+			return UE_REAL_TO_FLOAT(Domain.Y);
 		}
 
 		//Handle edge case of NaN
 		if (FMath::IsNaN(InFrequencyValue))
 		{
-			return Domain.X;
+			return UE_REAL_TO_FLOAT(Domain.X);
 		}
 
 		const FVector2D RangeLog(FMath::Max(FMath::Loge(Range.X), SMALL_NUMBER), FMath::Min(FMath::Loge(Range.Y), BIG_NUMBER));
 		const FVector2D::FReal FrequencyLog = FMath::Loge(InFrequencyValue);
-		return FMath::GetMappedRangeValueUnclamped(RangeLog, Domain, FrequencyLog);
+		return UE_REAL_TO_FLOAT(FMath::GetMappedRangeValueUnclamped(RangeLog, Domain, FrequencyLog));
 	}
 
 	// Using midi tuning standard, compute midi from frequency in hz
@@ -494,15 +494,15 @@ namespace Audio
 			{
 				if (i != j)
 				{
-					float Denom = Points[i].X - Points[j].X;
+					float Denom = UE_REAL_TO_FLOAT(Points[i].X - Points[j].X);
 					if (FMath::Abs(Denom) < SMALL_NUMBER)
 					{
 						Denom = SMALL_NUMBER;
 					}
-					Lagrangian *= (Alpha - Points[j].X) / Denom;
+					Lagrangian *= (Alpha - UE_REAL_TO_FLOAT(Points[j].X)) / Denom;
 				}
 			}
-			Output += Lagrangian * Points[i].Y;
+			Output += Lagrangian * UE_REAL_TO_FLOAT(Points[i].Y);
 		}
 		return Output;
 	}
diff --git a/Engine/Source/Runtime/SlateCore/Public/Input/Events.h b/Engine/Source/Runtime/SlateCore/Public/Input/Events.h
--- a/Engine/Source/Runtime/SlateCore/Public/Input/Events.h
+++ b/Engine/Source/Runtime/SlateCore/Public/Input/Events.h
@@ -841,7 +841,7 @@ public:
 	FKey GetEffectingButton() const { return EffectingButton; }
 	
 	/** How much did the mouse wheel turn since the last mouse event */
-	float GetWheelDelta() const { return WheelOrGestureDelta.Y; }
+	float GetWheelDelta() const { return UE_REAL_TO_FLOAT(WheelOrGestureDelta.Y); }
 
 	/** Returns the index of the user that caused the event */
 	int32 GetUserIndex() const { return UserIndex; }
diff --git a/Engine/Source/Runtime/SlateCore/Public/Layout/Margin.h b/Engine/Source/Runtime/SlateCore/Public/Layout/Margin.h
--- a/Engine/Source/Runtime/SlateCore/Public/Layout/Margin.h
+++ b/Engine/Source/Runtime/SlateCore/Public/Layout/Margin.h
@@ -64,10 +64,10 @@ public:
 
 	/** Construct a Margin where Horizontal describes Left and Right spacing while Vertical describes Top and Bottom spacing */
 	FMargin(const FVector2D& InVector)
-		: Left(InVector.X)
-		, Top(InVector.Y)
-		, Right(InVector.X)
-		, Bottom(InVector.Y)
+		: Left(UE_REAL_TO_FLOAT(InVector.X))
+		, Top(UE_REAL_TO_FLOAT(InVector.Y))
+		, Right(UE_REAL_TO_FLOAT(InVector.X))
+		, Bottom(UE_REAL_TO_FLOAT(InVector.Y))
 	{ }
 	
 	/** Construct a Margin where the spacing on each side is individually specified. */
diff --git a/Engine/Source/Runtime/SlateCore/Public/Layout/SlateRect.h b/Engine/Source/Runtime/SlateCore/Public/Layout/SlateRect.h
--- a/Engine/Source/Runtime/SlateCore/Public/Layout/SlateRect.h
+++ b/Engine/Source/Runtime/SlateCore/Public/Layout/SlateRect.h
@@ -38,10 +38,10 @@ public:
 	{ }
 
 	FSlateRect( const FVector2D& InStartPos, const FVector2D& InEndPos )
-		: Left(InStartPos.X)
-		, Top(InStartPos.Y)
-		, Right(InEndPos.X)
-		, Bottom(InEndPos.Y)
+		: Left(UE_REAL_TO_FLOAT(InStartPos.X))
+		, Top(UE_REAL_TO_FLOAT(InStartPos.Y))
+		, Right(UE_REAL_TO_FLOAT(InEndPos.X))
+		, Bottom(UE_REAL_TO_FLOAT(InEndPos.Y))
 	{ }
 
 	/**
