Deinterleaved DSP for metasound wave player

diff --git a/Engine/Source/Runtime/SignalProcessing/Private/ConvertDeinterleave.cpp b/Engine/Source/Runtime/SignalProcessing/Private/ConvertDeinterleave.cpp
new file mode 100644
--- /dev/null
+++ b/Engine/Source/Runtime/SignalProcessing/Private/ConvertDeinterleave.cpp
@@ -0,0 +1,259 @@
+// Copyright Epic Games, Inc. All Rights Reserved.
+
+#include "DSP/ConvertDeinterleave.h"
+
+#include "Containers/Array.h"
+#include "Containers/ArrayView.h"
+#include "DSP/ChannelMap.h"
+#include "DSP/MultichannelBuffer.h"
+
+namespace Audio
+{
+	namespace ConvertDeinterleavePrivate
+	{
+		// Template class for generic deinterleave/convert operation.  It uses
+		// mixing gains provided by `Audio::Create2DChannelMap` which are AC3 
+		// downmixing values. 
+		//
+		// Using a template class with the channel counts as template arguments
+		// allows the compiler to use  hard coded channel counts when compiling 
+		// the sample loop and/or for optimizations to be hand introduced using 
+		// class template specialization.
+		template<int32 NumInputChannels, int32 NumOutputChannels>
+		struct TConvertDeinterleave : IConvertDeinterleave
+		{
+			static_assert(NumInputChannels > 0);
+			static_assert(NumOutputChannels > 0);
+
+			TConvertDeinterleave()
+			{
+				using namespace Audio;
+
+				FChannelMapParams Params;
+				Params.NumInputChannels = NumInputChannels;
+				Params.NumOutputChannels = NumOutputChannels;
+				Params.Order = EChannelMapOrder::InputMajorOrder;
+				Params.MonoUpmixMethod = EChannelMapMonoUpmixMethod::EqualPower;
+				Params.bIsCenterChannelOnly = false;
+
+				bool bSuccess = Create2DChannelMap(Params, ChannelGains);
+				check(bSuccess);
+				checkf(ChannelGains.Num() == (NumInputChannels * NumOutputChannels), TEXT("Expected channel array of %d elements. Actual num elements: %d"), (NumInputChannels * NumOutputChannels), ChannelGains.Num());
+			}
+
+			virtual ~TConvertDeinterleave() = default;
+
+			const float* GetChannelGains(int32 InOutputChannelIndex) const
+			{
+				return &ChannelGains.GetData()[InOutputChannelIndex * NumInputChannels];
+			}
+
+			virtual void ProcessAudio(TArrayView<const float> InSamples, FMultichannelBuffer& OutSamples) const override
+			{
+				checkf((InSamples.Num() % NumInputChannels) == 0, TEXT("Input sample buffer contains partial audio frame."));
+				checkf((InSamples.Num() / NumInputChannels) == GetMultichannelBufferNumFrames(OutSamples), TEXT("Audio buffer frame count mismatch."));
+				checkf(NumOutputChannels == OutSamples.Num(), TEXT("Output audio buffer not initialized to expected channel count."));
+
+				const int32 NumFrames = GetMultichannelBufferNumFrames(OutSamples);
+				const float* InSampleData = InSamples.GetData();
+
+				float* OutputChannelDataPtrs[NumOutputChannels];
+				for (int32 OutChannelIndex = 0; OutChannelIndex < NumOutputChannels; OutChannelIndex++)
+				{
+					OutputChannelDataPtrs[OutChannelIndex] = OutSamples[OutChannelIndex].GetData();
+				}
+
+				for (int32 Frame = 0; Frame < NumFrames; Frame++)
+				{
+					const float* InSampleFrameStart = &InSampleData[Frame * NumInputChannels];
+					for (int32 OutChannelIndex = 0; OutChannelIndex < NumOutputChannels; OutChannelIndex++)
+					{
+						float Value = 0.f;
+						const float* Gains = GetChannelGains(OutChannelIndex);
+
+						// Apply gains to interleaved input samples for a single output
+						// sample.
+						for (int32 InChannelIndex = 0; InChannelIndex < NumInputChannels; InChannelIndex++)
+						{
+							Value += InSampleFrameStart[InChannelIndex] * Gains[InChannelIndex];
+						}
+
+						*OutputChannelDataPtrs[OutChannelIndex] = Value;
+						OutputChannelDataPtrs[OutChannelIndex]++;
+					}
+				}
+			}
+		private:
+			TArray<float> ChannelGains;
+		};
+
+		// Specialization of TConvertDeinterleave for mono sources. This avoids
+		// the need for deinterleaving.
+		template<int32 NumOutputChannels>
+		struct TConvertDeinterleave<1, NumOutputChannels> : IConvertDeinterleave
+		{
+			static_assert(NumOutputChannels > 0);
+
+			TConvertDeinterleave()
+			{
+				using namespace Audio;
+
+				FChannelMapParams Params;
+				Params.NumInputChannels = 1;
+				Params.NumOutputChannels = NumOutputChannels;
+				Params.Order = EChannelMapOrder::InputMajorOrder;
+				Params.MonoUpmixMethod = EChannelMapMonoUpmixMethod::EqualPower;
+				Params.bIsCenterChannelOnly = false;
+
+				bool bSuccess = Create2DChannelMap(Params, ChannelGains);
+				check(bSuccess);
+				checkf(ChannelGains.Num() == NumOutputChannels, TEXT("Expected channel array of %d elements. Actual num elements: %d"), NumOutputChannels, ChannelGains.Num());
+			}
+
+			virtual ~TConvertDeinterleave() = default;
+
+			virtual void ProcessAudio(TArrayView<const float> InSamples, FMultichannelBuffer& OutSamples) const override
+			{
+				checkf(InSamples.Num() == GetMultichannelBufferNumFrames(OutSamples), TEXT("Audio buffer frame count mismatch."));
+				checkf(NumOutputChannels == OutSamples.Num(), TEXT("Output audio buffer not initialized to expected channel count."));
+
+				const int32 NumFrames = GetMultichannelBufferNumFrames(OutSamples);
+				const float* InSampleData = InSamples.GetData();
+
+				const int32 DataSize = NumFrames * sizeof(float);
+				for (int32 OutChannelIndex = 0; OutChannelIndex < NumOutputChannels; OutChannelIndex++)
+				{
+					const float Gain = ChannelGains[OutChannelIndex];
+					if (Gain != 0.f)
+					{
+						// Only need to multiply if gain is non-zero
+						Audio::BufferMultiplyByConstant(InSampleData, ChannelGains[OutChannelIndex], OutSamples[OutChannelIndex].GetData(), NumFrames);
+					}
+					else
+					{
+						// If gain is zero, we can set the output channel to zero.
+						FMemory::Memset(OutSamples[OutChannelIndex].GetData(), 0, NumFrames * sizeof(float));
+					}
+				}
+			}
+
+		private:
+			TArray<float> ChannelGains;
+		};
+
+		// Template specialization for converting mono to mono. No need to adjust
+		// gains or deinterleave.
+		template<>
+		struct TConvertDeinterleave<1, 1> : IConvertDeinterleave
+		{
+			virtual ~TConvertDeinterleave() = default;
+			virtual void ProcessAudio(TArrayView<const float> InAudio, FMultichannelBuffer& OutAudio) const override
+			{
+				checkf(InAudio.Num() == GetMultichannelBufferNumFrames(OutAudio), TEXT("Audio buffer frame count mismatch."));
+				checkf(1 == OutAudio.Num(), TEXT("Output audio buffer not initialized to expected channel count."));
+
+				const int32 NumFrames = GetMultichannelBufferNumFrames(OutAudio);
+				const float* InSampleData = InAudio.GetData();
+				FMemory::Memcpy(OutAudio[0].GetData(), InSampleData, NumFrames * sizeof(float));
+			}
+		};
+
+		template<int32 NumInputChannels>
+		TUniquePtr<IConvertDeinterleave> CreateConvertDeinterleave(int32 InNumOutputChannels)
+		{
+			// IConvertDeinterleave defines conversion operations channel counts 
+			// between 1 and 8. This range mirrors the supported channel map channel
+			// counts. If the supported set of channel maps is altered, the supported
+			// set of defined TConvertDeinterleave classes should also be updated.
+			static_assert(ChannelMapMaxNumChannels == 8);
+
+			// Find the appropriate instantiation given the number of input and output channels.
+			switch (InNumOutputChannels)
+			{
+				case 1:
+					return MakeUnique<TConvertDeinterleave<NumInputChannels, 1>>();
+					break;
+				case 2:
+					return MakeUnique<TConvertDeinterleave<NumInputChannels, 2>>();
+					break;
+				case 3:
+					return MakeUnique<TConvertDeinterleave<NumInputChannels, 3>>();
+					break;
+				case 4:
+					return MakeUnique<TConvertDeinterleave<NumInputChannels, 4>>();
+					break;
+				case 5:
+					return MakeUnique<TConvertDeinterleave<NumInputChannels, 5>>();
+					break;
+				case 6:
+					return MakeUnique<TConvertDeinterleave<NumInputChannels, 6>>();
+					break;
+				case 7:
+					return MakeUnique<TConvertDeinterleave<NumInputChannels, 7>>();
+					break;
+				case 8:
+					return MakeUnique<TConvertDeinterleave<NumInputChannels, 8>>();
+					break;
+				default:
+				{
+					checkf(false, TEXT("Unsupported output channel count: %d"), InNumOutputChannels);
+					return TUniquePtr<IConvertDeinterleave>(nullptr);
+				}
+			}
+		}
+	}
+
+	
+	TUniquePtr<IConvertDeinterleave> IConvertDeinterleave::Create(int32 InNumInputChannels, int32 InNumOutputChannels)
+	{
+		using namespace ConvertDeinterleavePrivate;
+
+		// IConvertDeinterleave defines conversion operations channel counts 
+		// between 1 and 8. This range mirrors the supported channel map channel
+		// counts. If the supported set of channel maps is altered, the supported
+		// set of defined TConvertDeinterleave classes should also be updated.
+		static_assert(ChannelMapMaxNumChannels == 8);
+
+		// Find the appropriate instantiation given the number of input and output channels.
+		switch (InNumInputChannels)
+		{
+			case 1:
+				return CreateConvertDeinterleave<1>(InNumOutputChannels);
+				break;
+
+			case 2:
+				return CreateConvertDeinterleave<2>(InNumOutputChannels);
+				break;
+
+			case 3:
+				return CreateConvertDeinterleave<3>(InNumOutputChannels);
+				break;
+
+			case 4:
+				return CreateConvertDeinterleave<4>(InNumOutputChannels);
+				break;
+
+			case 5:
+				return CreateConvertDeinterleave<5>(InNumOutputChannels);
+				break;
+
+			case 6:
+				return CreateConvertDeinterleave<6>(InNumOutputChannels);
+				break;
+
+			case 7:
+				return CreateConvertDeinterleave<7>(InNumOutputChannels);
+				break;
+
+			case 8:
+				return CreateConvertDeinterleave<8>(InNumOutputChannels);
+				break;
+
+			default:
+			{
+				checkf(false, TEXT("Unsupported input channel count: %d"), InNumInputChannels);
+				return TUniquePtr<IConvertDeinterleave>(nullptr);
+			}
+		}
+	}
+}
diff --git a/Engine/Source/Runtime/SignalProcessing/Private/MultichannelBuffer.cpp b/Engine/Source/Runtime/SignalProcessing/Private/MultichannelBuffer.cpp
new file mode 100644
--- /dev/null
+++ b/Engine/Source/Runtime/SignalProcessing/Private/MultichannelBuffer.cpp
@@ -0,0 +1,112 @@
+// Copyright Epic Games, Inc. All Rights Reserved.
+
+#include "DSP/MultichannelBuffer.h"
+
+#include "Algo/AllOf.h"
+#include "Containers/Array.h"
+#include "DSP/BufferVectorOperations.h"
+#include "DSP/Dsp.h"
+
+namespace Audio
+{
+	namespace MultichannelBufferPrivate
+	{
+		template<typename BufferType>
+		int32 GetMultichannelBufferNumFrames(const TArray<BufferType>& InBuffer)
+		{
+			if (InBuffer.Num() > 0)
+			{
+				const int32 NumFrames = InBuffer[0].Num();
+				
+				// Check that all of the buffers contain the same frame count.
+				checkf(Algo::AllOf(InBuffer, [&NumFrames](const BufferType& Buffer) { return NumFrames == Buffer.Num(); }), TEXT("Buffers contain inconsistent frame sizes"));
+
+				return NumFrames;
+			}
+
+			return 0;
+		}
+	}
+
+	void SetMultichannelBufferSize(int32 InNumChannels, int32 InNumFrames, FMultichannelBuffer& OutBuffer)
+	{
+		OutBuffer.SetNum(InNumChannels, false /* bAllowShrinking */);
+		for (FAlignedFloatBuffer& Buffer : OutBuffer)
+		{
+			Buffer.SetNumUninitialized(InNumFrames, false /* bAllowShrinking */);
+		}
+	}
+
+	void SetMultichannelCircularBufferCapacity(int32 InNumChannels, int32 InNumFrames, FMultichannelCircularBuffer& OutBuffer)
+	{
+		OutBuffer.SetNum(InNumChannels, false /* bAllowShrinking */);
+		for (TCircularAudioBuffer<float>& Buffer : OutBuffer)
+		{
+			Buffer.SetCapacity(InNumFrames);
+		}
+	}
+
+	int32 GetMultichannelBufferNumFrames(const FMultichannelBuffer& InBuffer)
+	{
+		return MultichannelBufferPrivate::GetMultichannelBufferNumFrames(InBuffer);
+	}
+
+	int32 GetMultichannelBufferNumFrames(const FMultichannelCircularBuffer& InBuffer)
+	{
+		return MultichannelBufferPrivate::GetMultichannelBufferNumFrames(InBuffer);
+	}
+
+	int32 GetMultichannelBufferNumFrames(const FMultichannelBufferView& InBuffer)
+	{
+		return MultichannelBufferPrivate::GetMultichannelBufferNumFrames(InBuffer);
+	}
+
+	FMultichannelBufferView MakeMultichannelBufferView(FMultichannelBuffer& InBuffer)
+	{
+		FMultichannelBufferView View;
+
+		View.Reset(InBuffer.Num());
+		for (FAlignedFloatBuffer& ChannelBuffer: InBuffer)
+		{
+			View.Emplace(ChannelBuffer);
+		}
+
+		return View;
+	}
+
+	FMultichannelBufferView MakeMultichannelBufferView(FMultichannelBuffer& InBuffer, int32 InStartFrameIndex, int32 InNumFrames)
+	{
+		FMultichannelBufferView View;
+
+		View.Reset(InBuffer.Num());
+		for (FAlignedFloatBuffer& ChannelBuffer: InBuffer)
+		{
+			check(InBuffer.Num() <= (InStartFrameIndex + InNumFrames));
+			View.Emplace(&ChannelBuffer.GetData()[InStartFrameIndex], InNumFrames);
+		}
+
+		return View;
+	}
+
+	FMultichannelBufferView SliceMultichannelBufferView(const FMultichannelBufferView& View, int32 InStartFrameIndex, int32 InNumFrames)
+	{
+		FMultichannelBufferView OutView;
+		for (int32 ChannelIndex = 0; ChannelIndex < View.Num(); ChannelIndex++)
+		{
+			OutView.Emplace(View[ChannelIndex].Slice(InStartFrameIndex, InNumFrames));
+		}
+		return OutView;
+	}
+
+	void ShiftMultichannelBufferView(int32 InNumFrames, FMultichannelBufferView& View)
+	{
+		int32 NumRemainingFrames = GetMultichannelBufferNumFrames(View);
+		int32 NumFramesToShift = FMath::Clamp(InNumFrames, 0, NumRemainingFrames);
+		int32 NumFramesToKeep = NumRemainingFrames - NumFramesToShift;
+
+		for (int32 ChannelIndex = 0; ChannelIndex < View.Num(); ChannelIndex++)
+		{
+			View[ChannelIndex] = View[ChannelIndex].Slice(NumFramesToShift, NumFramesToKeep);
+		}
+	}
+}
diff --git a/Engine/Source/Runtime/SignalProcessing/Public/DSP/ConvertDeinterleave.h b/Engine/Source/Runtime/SignalProcessing/Public/DSP/ConvertDeinterleave.h
new file mode 100644
--- /dev/null
+++ b/Engine/Source/Runtime/SignalProcessing/Public/DSP/ConvertDeinterleave.h
@@ -0,0 +1,45 @@
+// Copyright Epic Games, Inc. All Rights Reserved.
+
+#pragma once
+
+#include "DSP/MultichannelBuffer.h"
+#include "Containers/ArrayView.h"
+#include "Templates/UniquePtr.h"
+
+namespace Audio
+{
+	/** IConvertDeinterleave is an interface for transforming multichannel interleaved
+	 * audio samples into multichannel deinterleaved samples. The channel count of 
+	 * the input and output audio may differ. 
+	 *
+	 * The deinterleaving and channel format conversion operations are combined
+	 * into this single object as both operations are often required for any given 
+	 * source audio. 
+	 */
+	struct IConvertDeinterleave
+	{
+		virtual ~IConvertDeinterleave() = default;
+
+		/** Deinterleave and convert the channel format of the input audio.
+		 *
+		 * @param InSamples - ArrayView of interleaved input samples. The number of 
+		 *                    samples must be evenly divisible by the number of 
+		 *                    input channels. 
+		 * @param OutSamples - A multichannel buffer of samples generated by deinterleaving
+		 *                     the input samples and possibly upmixing or downmixing
+		 *                     them to the target channel count. 
+		 */
+		virtual void ProcessAudio(TArrayView<const float> InSamples, FMultichannelBuffer& OutSamples) const = 0;
+
+
+		/** Create an IConvertDeinterleave object for a given number of input
+		 * and output channels. 
+		 * 
+		 * @param InNumInputChannels - Number of input channels in incoming audio.
+		 * @param InNumOutputChannels - Number of output channels in outgoing audio.
+		 *
+		 * @return A valid TUniquePtr<> on success, an invalid TUniquePtr<> on failure. 
+		 */
+		SIGNALPROCESSING_API static TUniquePtr<IConvertDeinterleave> Create(int32 InNumInputChannels, int32 InNumOutputChannels);
+	};
+}
diff --git a/Engine/Source/Runtime/SignalProcessing/Public/DSP/Dsp.h b/Engine/Source/Runtime/SignalProcessing/Public/DSP/Dsp.h
--- a/Engine/Source/Runtime/SignalProcessing/Public/DSP/Dsp.h
+++ b/Engine/Source/Runtime/SignalProcessing/Public/DSP/Dsp.h
@@ -868,7 +868,7 @@ namespace Audio
 		}
 
 		/** Push an array of values into circular buffer. */
-		int32 Push(TArrayView<SampleType> InBuffer)
+		int32 Push(TArrayView<const SampleType> InBuffer)
 		{
 			return Push(InBuffer.GetData(), InBuffer.Num());
 		}
@@ -1002,6 +1002,19 @@ namespace Audio
 			return NumSamplesRead;
 		}
 
+		// Pops some amount of samples into this circular buffer.
+		// Returns the amount of samples read.
+		int32 Pop(uint32 NumSamples)
+		{
+			check(NumSamples < ((uint32)TNumericLimits<int32>::Max()));
+
+			int32 NumSamplesRead = FMath::Min<int32>(NumSamples, Num());
+
+			ReadCounter.Set((ReadCounter.GetValue() + NumSamplesRead) % Capacity);
+
+			return NumSamplesRead;
+		}
+
 		// Pops a single element.
 		// Will assert if the buffer is empty. Please check Num() > 0 before calling.
 		SampleType Pop()
diff --git a/Engine/Source/Runtime/SignalProcessing/Public/DSP/MultichannelBuffer.h b/Engine/Source/Runtime/SignalProcessing/Public/DSP/MultichannelBuffer.h
new file mode 100644
--- /dev/null
+++ b/Engine/Source/Runtime/SignalProcessing/Public/DSP/MultichannelBuffer.h
@@ -0,0 +1,69 @@
+// Copyright Epic Games, Inc. All Rights Reserved.
+
+#pragma once
+
+#include "Containers/Array.h"
+#include "DSP/BufferVectorOperations.h"
+#include "DSP/Dsp.h"
+
+namespace Audio
+{
+	/** A deinterleaved multichannel buffer */
+	using FMultichannelBuffer = TArray<Audio::FAlignedFloatBuffer>;
+
+	/** A deinterleaved multichannel buffer view. The semantics are similar to 
+	 * TArray<> and TArrayView<>. FMultichannelBufferView allows sample values to 
+	 * be read and written, but it cannot be resized. 
+	 */
+	using FMultichannelBufferView = TArray<TArrayView<float>>;
+
+	/** A deinterleaved multichannel circular buffer */
+	using FMultichannelCircularBuffer = TArray<Audio::TCircularAudioBuffer<float>>;
+
+	/** Set the number of channels and frames for a multichannel buffer.
+	 *
+	 * @param InNumChannels - Number of channels to hold in the buffer.
+	 * @param InNumFrames - Number of frames to hold in the buffer.
+	 * @param OutBuffer - Buffer to resize.
+	 */
+	SIGNALPROCESSING_API void SetMultichannelBufferSize(int32 InNumChannels, int32 InNumFrames, FMultichannelBuffer& OutBuffer);
+
+	/** Set the number of channels and capacity for a multichannel circular buffer.
+	 *
+	 * @param InNumChannels - Number of channels to hold in the buffer.
+	 * @param InNumFrames - Maximum number of frames the buffer can hold.
+	 * @param OutBuffer - Buffer to resize.
+	 */
+	SIGNALPROCESSING_API void SetMultichannelCircularBufferCapacity(int32 InNumChannels, int32 InNumFrames, FMultichannelCircularBuffer& OutBuffer);
+
+	/** Set the number of frames for a multichannel buffer.
+	 *
+	 * @param InNumFrames - Number of frames to hold in the buffer.
+	 * @param OutBuffer - Buffer to resize.
+	 */
+	SIGNALPROCESSING_API void SetMultichannelBufferSize(int32 InNumFrames, FMultichannelBuffer& OutBuffer);
+
+	/** Return the number of frames in the buffer. It is expected that each channel
+	 * contains the same number of frames. */
+	SIGNALPROCESSING_API int32 GetMultichannelBufferNumFrames(const FMultichannelBuffer& InBuffer);
+
+	/** Return the number of frames in the buffer. It is expected that each channel
+	 * contains the same number of frames. */
+	SIGNALPROCESSING_API int32 GetMultichannelBufferNumFrames(const FMultichannelCircularBuffer& InBuffer);
+
+	/** Return the number of frames in the buffer. It is expected that each channel
+	 * contains the same number of frames. */
+	SIGNALPROCESSING_API int32 GetMultichannelBufferNumFrames(const FMultichannelBufferView& InBuffer);
+
+	/** Creates a FMultichannelBufferView from a FMultichannelBuffer. */
+	SIGNALPROCESSING_API FMultichannelBufferView MakeMultichannelBufferView(FMultichannelBuffer& InBuffer);
+
+	/** Creates a FMultichannelBufferView from a FMultichannelBuffer. */
+	SIGNALPROCESSING_API FMultichannelBufferView MakeMultichannelBufferView(FMultichannelBuffer& InBuffer, int32 InStartFrameIndex, int32 InNumFrames);
+
+	/** Slice a multichannel buffer view to a given frame range. */
+	SIGNALPROCESSING_API FMultichannelBufferView SliceMultichannelBufferView(const FMultichannelBufferView& View, int32 InStartFrameIndex, int32 InNumFrames);
+
+	/** Shift the start frame of a multichannel buffer view to a given number of frames. */
+	SIGNALPROCESSING_API void ShiftMultichannelBufferView(int32 InNumFrames, FMultichannelBufferView& View);
+}
