Fix for altered volume when using a mono wave file in a metasound wave player node.

diff --git a/Engine/Plugins/Runtime/Metasound/Source/MetasoundEngine/Private/MetasoundWavePlayerNode.cpp b/Engine/Plugins/Runtime/Metasound/Source/MetasoundEngine/Private/MetasoundWavePlayerNode.cpp
--- a/Engine/Plugins/Runtime/Metasound/Source/MetasoundEngine/Private/MetasoundWavePlayerNode.cpp
+++ b/Engine/Plugins/Runtime/Metasound/Source/MetasoundEngine/Private/MetasoundWavePlayerNode.cpp
@@ -688,7 +688,15 @@ namespace Metasound
 						NumDeinterleaveChannels = NumOutputChannels;
 
 						// Create algorithm for channel conversion and deinterleave 
-						ConvertDeinterleave = Audio::IConvertDeinterleave::Create(WaveProxyReader->GetNumChannels(), NumDeinterleaveChannels);
+						Audio::FConvertDeinterleaveParams ConvertDeinterleaveParams;
+						ConvertDeinterleaveParams.NumInputChannels = WaveProxyReader->GetNumChannels();
+						ConvertDeinterleaveParams.NumOutputChannels = NumDeinterleaveChannels;
+						// Original implementation of MetaSound WavePlayer upmixed 
+						// mono using FullVolume. In the future, the mono upmix 
+						// method may be exposed as a node input to facilitate 
+						// better control.
+						ConvertDeinterleaveParams.MonoUpmixMethod = Audio::EChannelMapMonoUpmixMethod::FullVolume;
+						ConvertDeinterleave = Audio::IConvertDeinterleave::Create(ConvertDeinterleaveParams);
 						Audio::SetMultichannelBufferSize(NumDeinterleaveChannels, DeinterleaveBlockSizeInFrames, DeinterleavedBuffer);
 
 						// Initialize source buffer
diff --git a/Engine/Source/Runtime/SignalProcessing/Private/ConvertDeinterleave.cpp b/Engine/Source/Runtime/SignalProcessing/Private/ConvertDeinterleave.cpp
--- a/Engine/Source/Runtime/SignalProcessing/Private/ConvertDeinterleave.cpp
+++ b/Engine/Source/Runtime/SignalProcessing/Private/ConvertDeinterleave.cpp
@@ -25,18 +25,21 @@ namespace Audio
 			static_assert(NumInputChannels > 0);
 			static_assert(NumOutputChannels > 0);
 
-			TConvertDeinterleave()
+			TConvertDeinterleave(const FConvertDeinterleaveParams& InParams)
 			{
 				using namespace Audio;
 
-				FChannelMapParams Params;
-				Params.NumInputChannels = NumInputChannels;
-				Params.NumOutputChannels = NumOutputChannels;
-				Params.Order = EChannelMapOrder::InputMajorOrder;
-				Params.MonoUpmixMethod = EChannelMapMonoUpmixMethod::EqualPower;
-				Params.bIsCenterChannelOnly = false;
+				check(InParams.NumInputChannels == NumInputChannels);
+				check(InParams.NumOutputChannels == NumOutputChannels);
 
-				bool bSuccess = Create2DChannelMap(Params, ChannelGains);
+				FChannelMapParams ChannelMapParams;
+				ChannelMapParams.NumInputChannels = NumInputChannels;
+				ChannelMapParams.NumOutputChannels = NumOutputChannels;
+				ChannelMapParams.Order = EChannelMapOrder::InputMajorOrder;
+				ChannelMapParams.MonoUpmixMethod = InParams.MonoUpmixMethod;
+				ChannelMapParams.bIsCenterChannelOnly = false;
+
+				bool bSuccess = Create2DChannelMap(ChannelMapParams, ChannelGains);
 				check(bSuccess);
 				checkf(ChannelGains.Num() == (NumInputChannels * NumOutputChannels), TEXT("Expected channel array of %d elements. Actual num elements: %d"), (NumInputChannels * NumOutputChannels), ChannelGains.Num());
 			}
@@ -94,18 +97,21 @@ namespace Audio
 		{
 			static_assert(NumOutputChannels > 0);
 
-			TConvertDeinterleave()
+			TConvertDeinterleave(const FConvertDeinterleaveParams& InParams)
 			{
 				using namespace Audio;
 
-				FChannelMapParams Params;
-				Params.NumInputChannels = 1;
-				Params.NumOutputChannels = NumOutputChannels;
-				Params.Order = EChannelMapOrder::InputMajorOrder;
-				Params.MonoUpmixMethod = EChannelMapMonoUpmixMethod::EqualPower;
-				Params.bIsCenterChannelOnly = false;
+				check(InParams.NumInputChannels == 1);
+				check(InParams.NumOutputChannels == NumOutputChannels);
+
+				FChannelMapParams ChannelMapParams;
+				ChannelMapParams.NumInputChannels = 1;
+				ChannelMapParams.NumOutputChannels = NumOutputChannels;
+				ChannelMapParams.Order = EChannelMapOrder::InputMajorOrder;
+				ChannelMapParams.MonoUpmixMethod = InParams.MonoUpmixMethod;
+				ChannelMapParams.bIsCenterChannelOnly = false;
 
-				bool bSuccess = Create2DChannelMap(Params, ChannelGains);
+				bool bSuccess = Create2DChannelMap(ChannelMapParams, ChannelGains);
 				check(bSuccess);
 				checkf(ChannelGains.Num() == NumOutputChannels, TEXT("Expected channel array of %d elements. Actual num elements: %d"), NumOutputChannels, ChannelGains.Num());
 			}
@@ -146,6 +152,12 @@ namespace Audio
 		template<>
 		struct TConvertDeinterleave<1, 1> : IConvertDeinterleave
 		{
+			TConvertDeinterleave(const FConvertDeinterleaveParams& InParams)
+			{
+				check(InParams.NumInputChannels == 1);
+				check(InParams.NumOutputChannels == 1);
+			}
+
 			virtual ~TConvertDeinterleave() = default;
 			virtual void ProcessAudio(TArrayView<const float> InAudio, FMultichannelBuffer& OutAudio) const override
 			{
@@ -159,7 +171,7 @@ namespace Audio
 		};
 
 		template<int32 NumInputChannels>
-		TUniquePtr<IConvertDeinterleave> CreateConvertDeinterleave(int32 InNumOutputChannels)
+		TUniquePtr<IConvertDeinterleave> CreateConvertDeinterleave(const FConvertDeinterleaveParams& InParams)
 		{
 			// IConvertDeinterleave defines conversion operations channel counts 
 			// between 1 and 8. This range mirrors the supported channel map channel
@@ -168,35 +180,35 @@ namespace Audio
 			static_assert(ChannelMapMaxNumChannels == 8);
 
 			// Find the appropriate instantiation given the number of input and output channels.
-			switch (InNumOutputChannels)
+			switch (InParams.NumOutputChannels)
 			{
 				case 1:
-					return MakeUnique<TConvertDeinterleave<NumInputChannels, 1>>();
+					return MakeUnique<TConvertDeinterleave<NumInputChannels, 1>>(InParams);
 					break;
 				case 2:
-					return MakeUnique<TConvertDeinterleave<NumInputChannels, 2>>();
+					return MakeUnique<TConvertDeinterleave<NumInputChannels, 2>>(InParams);
 					break;
 				case 3:
-					return MakeUnique<TConvertDeinterleave<NumInputChannels, 3>>();
+					return MakeUnique<TConvertDeinterleave<NumInputChannels, 3>>(InParams);
 					break;
 				case 4:
-					return MakeUnique<TConvertDeinterleave<NumInputChannels, 4>>();
+					return MakeUnique<TConvertDeinterleave<NumInputChannels, 4>>(InParams);
 					break;
 				case 5:
-					return MakeUnique<TConvertDeinterleave<NumInputChannels, 5>>();
+					return MakeUnique<TConvertDeinterleave<NumInputChannels, 5>>(InParams);
 					break;
 				case 6:
-					return MakeUnique<TConvertDeinterleave<NumInputChannels, 6>>();
+					return MakeUnique<TConvertDeinterleave<NumInputChannels, 6>>(InParams);
 					break;
 				case 7:
-					return MakeUnique<TConvertDeinterleave<NumInputChannels, 7>>();
+					return MakeUnique<TConvertDeinterleave<NumInputChannels, 7>>(InParams);
 					break;
 				case 8:
-					return MakeUnique<TConvertDeinterleave<NumInputChannels, 8>>();
+					return MakeUnique<TConvertDeinterleave<NumInputChannels, 8>>(InParams);
 					break;
 				default:
 				{
-					checkf(false, TEXT("Unsupported output channel count: %d"), InNumOutputChannels);
+					checkf(false, TEXT("Unsupported output channel count: %d"), InParams.NumOutputChannels);
 					return TUniquePtr<IConvertDeinterleave>(nullptr);
 				}
 			}
@@ -204,7 +216,7 @@ namespace Audio
 	}
 
 	
-	TUniquePtr<IConvertDeinterleave> IConvertDeinterleave::Create(int32 InNumInputChannels, int32 InNumOutputChannels)
+	TUniquePtr<IConvertDeinterleave> IConvertDeinterleave::Create(const FConvertDeinterleaveParams& InParams)
 	{
 		using namespace ConvertDeinterleavePrivate;
 
@@ -215,43 +227,43 @@ namespace Audio
 		static_assert(ChannelMapMaxNumChannels == 8);
 
 		// Find the appropriate instantiation given the number of input and output channels.
-		switch (InNumInputChannels)
+		switch (InParams.NumInputChannels)
 		{
 			case 1:
-				return CreateConvertDeinterleave<1>(InNumOutputChannels);
+				return CreateConvertDeinterleave<1>(InParams);
 				break;
 
 			case 2:
-				return CreateConvertDeinterleave<2>(InNumOutputChannels);
+				return CreateConvertDeinterleave<2>(InParams);
 				break;
 
 			case 3:
-				return CreateConvertDeinterleave<3>(InNumOutputChannels);
+				return CreateConvertDeinterleave<3>(InParams);
 				break;
 
 			case 4:
-				return CreateConvertDeinterleave<4>(InNumOutputChannels);
+				return CreateConvertDeinterleave<4>(InParams);
 				break;
 
 			case 5:
-				return CreateConvertDeinterleave<5>(InNumOutputChannels);
+				return CreateConvertDeinterleave<5>(InParams);
 				break;
 
 			case 6:
-				return CreateConvertDeinterleave<6>(InNumOutputChannels);
+				return CreateConvertDeinterleave<6>(InParams);
 				break;
 
 			case 7:
-				return CreateConvertDeinterleave<7>(InNumOutputChannels);
+				return CreateConvertDeinterleave<7>(InParams);
 				break;
 
 			case 8:
-				return CreateConvertDeinterleave<8>(InNumOutputChannels);
+				return CreateConvertDeinterleave<8>(InParams);
 				break;
 
 			default:
 			{
-				checkf(false, TEXT("Unsupported input channel count: %d"), InNumInputChannels);
+				checkf(false, TEXT("Unsupported input channel count: %d"), InParams.NumInputChannels);
 				return TUniquePtr<IConvertDeinterleave>(nullptr);
 			}
 		}
diff --git a/Engine/Source/Runtime/SignalProcessing/Public/DSP/ConvertDeinterleave.h b/Engine/Source/Runtime/SignalProcessing/Public/DSP/ConvertDeinterleave.h
--- a/Engine/Source/Runtime/SignalProcessing/Public/DSP/ConvertDeinterleave.h
+++ b/Engine/Source/Runtime/SignalProcessing/Public/DSP/ConvertDeinterleave.h
@@ -2,12 +2,26 @@
 
 #pragma once
 
+#include "DSP/ChannelMap.h"
 #include "DSP/MultichannelBuffer.h"
 #include "Containers/ArrayView.h"
 #include "Templates/UniquePtr.h"
 
 namespace Audio
 {
+	/** Parameters for creating a IConvertDeinterleave object. */
+	struct FConvertDeinterleaveParams
+	{
+		// Number of channels in the input audio.
+		int32 NumInputChannels = 0;
+
+		// Number of channels in the output audio.
+		int32 NumOutputChannels = 0;
+
+		// Method for upmixing mono audio (only used if NumInputChannels == 1)
+		EChannelMapMonoUpmixMethod MonoUpmixMethod = EChannelMapMonoUpmixMethod::EqualPower;
+	};
+
 	/** IConvertDeinterleave is an interface for transforming multichannel interleaved
 	 * audio samples into multichannel deinterleaved samples. The channel count of 
 	 * the input and output audio may differ. 
@@ -40,6 +54,6 @@ namespace Audio
 		 *
 		 * @return A valid TUniquePtr<> on success, an invalid TUniquePtr<> on failure. 
 		 */
-		SIGNALPROCESSING_API static TUniquePtr<IConvertDeinterleave> Create(int32 InNumInputChannels, int32 InNumOutputChannels);
+		SIGNALPROCESSING_API static TUniquePtr<IConvertDeinterleave> Create(const FConvertDeinterleaveParams& InParams);
 	};
 }
