Temp fix for channel order issues for non-xaudio2 skews

diff --git a/Engine/Source/Runtime/SignalProcessing/Private/ChannelMap.cpp b/Engine/Source/Runtime/SignalProcessing/Private/ChannelMap.cpp
--- a/Engine/Source/Runtime/SignalProcessing/Private/ChannelMap.cpp
+++ b/Engine/Source/Runtime/SignalProcessing/Private/ChannelMap.cpp
@@ -136,6 +136,8 @@ namespace Audio
 			0.0f,			0.0f,		0.0f,		0.0f,			1.0f,		0.0f,		// SideRight
 		};
 
+#if PLATFORM_MICROSOFT
+
 		static constexpr float To7Point1Matrix[ChannelMapMaxNumChannels * 8] =
 		{
 			// FrontLeft	FrontRight	Center		LowFrequency	SideLeft	SideRight	BackLeft	BackRight
@@ -149,6 +151,24 @@ namespace Audio
 			0.0f,			0.0f,		0.0f,		0.0f,			0.0f,		0.0f,		0.0f,		1.0f,		// BackRight
 		};
 
+#else //PLATFORM_MICROSOFT
+
+		// NOTE: the BackLeft/BackRight and SideLeft/SideRight are reversed than they should be since our 7.1 importer code has it backward
+		static constexpr float To7Point1Matrix[ChannelMapMaxNumChannels * 8] =
+		{
+			// FrontLeft	FrontRight	Center		LowFrequency	SideLeft	SideRight	BackLeft	BackRight
+			1.0f,			0.0f,		0.0f,		0.0f,			0.0f,		0.0f,		0.0f,		0.0f,		// FrontLeft
+			0.0f,			1.0f,		0.0f,		0.0f,			0.0f,		0.0f,		0.0f,		0.0f,		// FrontRight
+			0.0f,			0.0f,		1.0f,		0.0f,			0.0f,		0.0f,		0.0f,		0.0f,		// FrontCenter
+			0.0f,			0.0f,		0.0f,		1.0f,			0.0f,		0.0f,		0.0f,		0.0f,		// LowFrequency
+			0.0f,			0.0f,		0.0f,		0.0f,			0.0f,		0.0f,		1.0f,		0.0f,		// BackLeft
+			0.0f,			0.0f,		0.0f,		0.0f,			0.0f,		0.0f,		0.0f,		1.0f,		// BackRight
+			0.0f,			0.0f,		0.0f,		0.0f,			1.0f,		0.0f,		0.0f,		0.0f,		// SideLeft
+			0.0f,			0.0f,		0.0f,		0.0f,			0.0f,		1.0f,		0.0f,		0.0f,		// SideRight
+		};
+
+#endif//PLATFORM_MICROSOFT
+
 		static constexpr float VorbisTo7Point1Matrix[ChannelMapVorbisNumChannels * 8] =
 		{
 			// FrontLeft	Center		FrontRight	SideLeft		SideRight	LowFrequency		
