Intesl ISPC optimizations for Audio
Submitting on behalf of Jeff Rous

diff --git a/Engine/Plugins/Experimental/MotoSynth/Source/MotoSynth/Private/MotoSynthEngine.cpp b/Engine/Plugins/Experimental/MotoSynth/Source/MotoSynth/Private/MotoSynthEngine.cpp
--- a/Engine/Plugins/Experimental/MotoSynth/Source/MotoSynth/Private/MotoSynthEngine.cpp
+++ b/Engine/Plugins/Experimental/MotoSynth/Source/MotoSynth/Private/MotoSynthEngine.cpp
@@ -492,14 +492,16 @@ void FMotoSynthEngine::GenerateGranularEngine(float* OutAudio, int32 NumSamples)
 		CurrentRPM += RPMDelta;
 	}
 
+	TArrayView<float> OutAudioView(OutAudio, NumSamples);
+
 	if (!FMath::IsNearlyEqual(TargetGranularEngineVolume, GranularEngineVolume))
 	{
-		Audio::FadeBufferFast(OutAudio, NumSamples, GranularEngineVolume, TargetGranularEngineVolume);
+		Audio::ArrayFade(OutAudioView, GranularEngineVolume, TargetGranularEngineVolume);
 		GranularEngineVolume = TargetGranularEngineVolume;
 	}
 	else if (!FMath::IsNearlyEqual(GranularEngineVolume, 1.0f))
 	{
-		Audio::MultiplyBufferByConstantInPlace(OutAudio, NumSamples, GranularEngineVolume);
+		Audio::ArrayMultiplyByConstantInPlace(OutAudioView, GranularEngineVolume);
 	}
 
 	// Apply the filter and mix the audio into the output buffer
diff --git a/Engine/Plugins/Experimental/MotoSynth/Source/MotoSynth/Private/MotoSynthSourceAsset.cpp b/Engine/Plugins/Experimental/MotoSynth/Source/MotoSynth/Private/MotoSynthSourceAsset.cpp
--- a/Engine/Plugins/Experimental/MotoSynth/Source/MotoSynth/Private/MotoSynthSourceAsset.cpp
+++ b/Engine/Plugins/Experimental/MotoSynth/Source/MotoSynth/Private/MotoSynthSourceAsset.cpp
@@ -382,7 +382,7 @@ void UMotoSynthSource::NormalizeForAnalysis()
 		// If we found one, which we should, use to normalize the signal
 		if (MaxSample > 0.0f)
 		{
-			Audio::MultiplyBufferByConstantInPlace(AnalysisBuffer, 1.0f / MaxSample);
+			Audio::ArrayMultiplyByConstantInPlace(AnalysisBuffer, 1.0f / MaxSample);
 		}
 	}
 }
diff --git a/Engine/Plugins/Runtime/AudioSynesthesia/Source/AudioSynesthesiaCore/Private/ConstantQNRTFactory.cpp b/Engine/Plugins/Runtime/AudioSynesthesia/Source/AudioSynesthesiaCore/Private/ConstantQNRTFactory.cpp
--- a/Engine/Plugins/Runtime/AudioSynesthesia/Source/AudioSynesthesiaCore/Private/ConstantQNRTFactory.cpp
+++ b/Engine/Plugins/Runtime/AudioSynesthesia/Source/AudioSynesthesiaCore/Private/ConstantQNRTFactory.cpp
@@ -170,11 +170,11 @@ namespace Audio
 				TAutoDeinterleaveView<float, FAudioBufferAlignedAllocator> DeinterleaveView(Window, ChannelBuffer, NumChannels);
 				for(auto Channel : DeinterleaveView)
 				{
-					MixInBufferFast(Channel.Values, MonoBuffer);		
+					ArrayMixIn(Channel.Values, MonoBuffer);		
 				}
 
 				// Equal power sum. assuming incoherrent signals.
-				MultiplyBufferByConstantInPlace(MonoBuffer, 1.f / FMath::Sqrt(static_cast<float>(NumChannels)));
+				ArrayMultiplyByConstantInPlace(MonoBuffer, 1.f / FMath::Sqrt(static_cast<float>(NumChannels)));
 
 				AnalyzeWindow(MonoBuffer, 0, *ConstantQResult);
 
diff --git a/Engine/Plugins/Runtime/AudioSynesthesia/Source/AudioSynesthesiaCore/Private/OnsetNRTFactory.cpp b/Engine/Plugins/Runtime/AudioSynesthesia/Source/AudioSynesthesiaCore/Private/OnsetNRTFactory.cpp
--- a/Engine/Plugins/Runtime/AudioSynesthesia/Source/AudioSynesthesiaCore/Private/OnsetNRTFactory.cpp
+++ b/Engine/Plugins/Runtime/AudioSynesthesia/Source/AudioSynesthesiaCore/Private/OnsetNRTFactory.cpp
@@ -188,11 +188,11 @@ namespace Audio
 				TAutoDeinterleaveView<float, FAudioBufferAlignedAllocator> DeinterleaveView(Window, ChannelBuffer, NumChannels);
 				for(auto Channel : DeinterleaveView)
 				{
-					MixInBufferFast(Channel.Values, MonoBuffer);		
+					ArrayMixIn(Channel.Values, MonoBuffer);		
 				}
 
 				// Equal power sum. assuming incoherrent signals.
-				MultiplyBufferByConstantInPlace(MonoBuffer, 1.f / FMath::Sqrt(static_cast<float>(NumChannels)));
+				ArrayMultiplyByConstantInPlace(MonoBuffer, 1.f / FMath::Sqrt(static_cast<float>(NumChannels)));
 
 				AnalyzeWindow(MonoBuffer, 0, *OnsetResult);
 			}
diff --git a/Engine/Plugins/Runtime/Metasound/Source/MetasoundFrontend/Public/MetasoundAudioBuffer.h b/Engine/Plugins/Runtime/Metasound/Source/MetasoundFrontend/Public/MetasoundAudioBuffer.h
--- a/Engine/Plugins/Runtime/Metasound/Source/MetasoundFrontend/Public/MetasoundAudioBuffer.h
+++ b/Engine/Plugins/Runtime/Metasound/Source/MetasoundFrontend/Public/MetasoundAudioBuffer.h
@@ -2,7 +2,7 @@
 
 #pragma once
 
-#include "DSP/BufferVectorOperations.h"
+#include "DSP/FloatArrayMath.h"
 #include "MetasoundGraphCoreModule.h"
 #include "MetasoundDataReferenceMacro.h"
 #include "MetasoundLog.h"
diff --git a/Engine/Plugins/Runtime/Metasound/Source/MetasoundStandardNodes/Private/MetasoundAdditiveSynthNode.cpp b/Engine/Plugins/Runtime/Metasound/Source/MetasoundStandardNodes/Private/MetasoundAdditiveSynthNode.cpp
--- a/Engine/Plugins/Runtime/Metasound/Source/MetasoundStandardNodes/Private/MetasoundAdditiveSynthNode.cpp
+++ b/Engine/Plugins/Runtime/Metasound/Source/MetasoundStandardNodes/Private/MetasoundAdditiveSynthNode.cpp
@@ -242,8 +242,11 @@ namespace Metasound
 				float NewRightGain = RightPanGain * BaseGain;
 
 				// Mix into output buffer
-				Audio::MixInBufferFast(ScratchBuffer.GetData(), LeftAudioOut->GetData(), NumSamples, NewLeftGain);
-				Audio::MixInBufferFast(ScratchBuffer.GetData(), RightAudioOut->GetData(), NumSamples, NewRightGain);
+				TArrayView<const float> ScratchBufferView(ScratchBuffer.GetData(), NumSamples);
+				TArrayView<float> LeftAudioOutView(LeftAudioOut->GetData(), NumSamples);
+				TArrayView<float> RightAudioOutView(RightAudioOut->GetData(), NumSamples);
+				Audio::ArrayMixIn(ScratchBufferView, LeftAudioOutView, NewLeftGain);
+				Audio::ArrayMixIn(ScratchBufferView, RightAudioOutView, NewRightGain);
 			}
 		}
 
diff --git a/Engine/Plugins/Runtime/Metasound/Source/MetasoundStandardNodes/Private/MetasoundCrossfadeNode.cpp b/Engine/Plugins/Runtime/Metasound/Source/MetasoundStandardNodes/Private/MetasoundCrossfadeNode.cpp
--- a/Engine/Plugins/Runtime/Metasound/Source/MetasoundStandardNodes/Private/MetasoundCrossfadeNode.cpp
+++ b/Engine/Plugins/Runtime/Metasound/Source/MetasoundStandardNodes/Private/MetasoundCrossfadeNode.cpp
@@ -106,7 +106,7 @@ namespace Metasound
 
 			// Zero the output buffer so we can mix into it
 			OutAudioBuffer.Zero();
-			float* OutAudioBufferPtr = OutAudioBuffer.GetData();
+			TArrayView<float> OutAudioBufferView(OutAudioBuffer.GetData(), OutAudioBuffer.Num());
 
 			// Now write to the scratch buffers w/ fade buffer fast given the new inputs
 			for (int32 i = 0; i < NumInputs; ++i)
@@ -116,10 +116,11 @@ namespace Metasound
 				{
 					// Copy the input to the output
 					const FAudioBufferReadRef& InBuff = InAudioBuffersValues[i];
+					TArrayView<const float> BufferView((*InBuff).GetData(), NumFramesPerBlock);
 					const float* BufferPtr = (*InBuff).GetData();
 
 					// mix in and fade to the target gain values
-					Audio::MixInBufferFast(BufferPtr, OutAudioBufferPtr, NumFramesPerBlock, PrevGains[i], CurrentGains[i]);
+					Audio::ArrayMixIn(BufferView, OutAudioBufferView, PrevGains[i], CurrentGains[i]);
 				}
 			}
 
diff --git a/Engine/Plugins/Runtime/Metasound/Source/MetasoundStandardNodes/Private/MetasoundITDPannerNode.cpp b/Engine/Plugins/Runtime/Metasound/Source/MetasoundStandardNodes/Private/MetasoundITDPannerNode.cpp
--- a/Engine/Plugins/Runtime/Metasound/Source/MetasoundStandardNodes/Private/MetasoundITDPannerNode.cpp
+++ b/Engine/Plugins/Runtime/Metasound/Source/MetasoundStandardNodes/Private/MetasoundITDPannerNode.cpp
@@ -198,6 +198,8 @@ namespace Metasound
 		int32 InputSampleCount = AudioInput->Num();
 		float* OutputLeftBufferPtr = AudioLeftOutput->GetData();
 		float* OutputRightBufferPtr = AudioRightOutput->GetData();
+		TArrayView<float> OutputLeftBufferView(AudioLeftOutput->GetData(), InputSampleCount);
+		TArrayView<float> OutputRightBufferView(AudioRightOutput->GetData(), InputSampleCount);
 
 		// Feed the input audio into the left and right delays
 		for (int32 i = 0; i < InputSampleCount; ++i)
@@ -209,13 +211,13 @@ namespace Metasound
 		// Now apply the panning
 		if (FMath::IsNearlyEqual(PrevLeftGain, CurrLeftDelay))
 		{
-			Audio::MultiplyBufferByConstantInPlace(OutputLeftBufferPtr, InputSampleCount, PrevLeftGain);
-			Audio::MultiplyBufferByConstantInPlace(OutputRightBufferPtr, InputSampleCount, PrevRightGain);
+			Audio::ArrayMultiplyByConstantInPlace(OutputLeftBufferView, PrevLeftGain);
+			Audio::ArrayMultiplyByConstantInPlace(OutputRightBufferView, PrevRightGain);
 		}
 		else
 		{
-			Audio::FadeBufferFast(OutputLeftBufferPtr, InputSampleCount, PrevLeftGain, CurrLeftGain);
-			Audio::FadeBufferFast(OutputRightBufferPtr, InputSampleCount, PrevRightGain, CurrRightGain);
+			Audio::ArrayFade(OutputLeftBufferView, PrevLeftGain, CurrLeftGain);
+			Audio::ArrayFade(OutputRightBufferView, PrevRightGain, CurrRightGain);
 
 			PrevLeftGain = CurrLeftGain;
 			PrevRightGain = CurrRightGain;
diff --git a/Engine/Plugins/Runtime/Metasound/Source/MetasoundStandardNodes/Private/MetasoundMathNodes.cpp b/Engine/Plugins/Runtime/Metasound/Source/MetasoundStandardNodes/Private/MetasoundMathNodes.cpp
--- a/Engine/Plugins/Runtime/Metasound/Source/MetasoundStandardNodes/Private/MetasoundMathNodes.cpp
+++ b/Engine/Plugins/Runtime/Metasound/Source/MetasoundStandardNodes/Private/MetasoundMathNodes.cpp
@@ -882,21 +882,18 @@ namespace Metasound
 					return;
 				}
 
-				const float* OperandData = InAdditionalOperands[i]->GetData();
-				float* OutData = OutResult->GetData();
+				TArrayView<const float> OperandDataView(InAdditionalOperands[i]->GetData(), NumSamples);
+				TArrayView<float> OutDataView(OutResult->GetData(), NumSamples);
 
 				if (NumSamples % AUDIO_NUM_FLOATS_PER_VECTOR_REGISTER)
 				{
-					for (int32 SampleIndex = 0; SampleIndex < NumSamples; ++SampleIndex)
-					{
-						OutData[SampleIndex] += OperandData[SampleIndex];
-					}
+					Audio::ArrayMixIn(OperandDataView, OutDataView);
 				}
 				else
 				{
 					const float StartGain = 1.0f;
 					const float EndGain = 1.0f;
-					Audio::MixInBufferFast(OperandData, OutData, NumSamples, StartGain, EndGain);
+					Audio::ArrayMixIn(OperandDataView, OutDataView, StartGain, EndGain);
 				}
 			}
 		}
@@ -1013,17 +1010,11 @@ namespace Metasound
 		{
 			FMemory::Memcpy(OutResult->GetData(), InPrimaryOperand->GetData(), sizeof(float) * OutResult->Num());
 
-			const int32 SIMDRemainder = OutResult->Num() % AUDIO_NUM_FLOATS_PER_VECTOR_REGISTER;
-			const int32 SIMDCount = OutResult->Num() - SIMDRemainder;
+			TArrayView<float> OutResultView(OutResult->GetData(), OutResult->Num());
 
 			for (int32 i = 0; i < InAdditionalOperands.Num(); ++i)
 			{
-				Audio::AddConstantToBufferInplace(OutResult->GetData(), SIMDCount, *InAdditionalOperands[i]);
-
-				for (int32 j = SIMDCount; j < OutResult->Num(); ++j)
-				{
-					OutResult->GetData()[j] += *InAdditionalOperands[i];
-				}
+				Audio::ArrayAddConstantInplace(OutResultView, *InAdditionalOperands[i]);
 			}
 		}
 	};
@@ -1099,20 +1090,10 @@ namespace Metasound
 					return;
 				}
 
-				const float* OperandData = InAdditionalOperands[i]->GetData();
-				float* OutData = OutResult->GetData();
+				TArrayView<const float> OperandDataView(InAdditionalOperands[i]->GetData(), NumSamples);
+				TArrayView<float> OutDataView(OutResult->GetData(), NumSamples);
 
-				if (NumSamples % AUDIO_NUM_FLOATS_PER_VECTOR_REGISTER)
-				{
-					for (int32 SampleIndex = 0; SampleIndex < NumSamples; ++SampleIndex)
-					{
-						OutData[i] -= OperandData[i];
-					}
-				}
-				else
-				{
-					Audio::BufferSubtractInPlace2Fast(OutData, OperandData, NumSamples);
-				}
+				Audio::ArraySubtractInPlace2(OutDataView, OperandDataView);
 			}
 		}
 	};
@@ -1252,20 +1233,10 @@ namespace Metasound
 					return;
 				}
 
-				const float* OperandData = InAdditionalOperands[i]->GetData();
-				float* OutData = OutResult->GetData();
+				TArrayView<const float> OperandDataView(InAdditionalOperands[i]->GetData(), NumSamples);
+				TArrayView<float> OutDataView(OutResult->GetData(), NumSamples);
 
-				if (NumSamples % AUDIO_NUM_FLOATS_PER_VECTOR_REGISTER)
-				{
-					for (int32 SampleIndex = 0; SampleIndex < NumSamples; ++SampleIndex)
-					{
-						OutData[SampleIndex] += OperandData[SampleIndex];
-					}
-				}
-				else
-				{
-					Audio::MultiplyBuffersInPlace(OperandData, OutData, NumSamples);
-				}
+				Audio::ArrayMultiplyInPlace(OperandDataView, OutDataView);
 			}
 		}
 	};
@@ -1337,7 +1308,9 @@ namespace Metasound
 				const int32 SIMDRemainder = OutResult->Num() % AUDIO_NUM_FLOATS_PER_VECTOR_REGISTER;
 				const int32 SIMDCount = OutResult->Num() - SIMDRemainder;
 
-				Audio::FadeBufferFast(OutResult->GetData(), SIMDCount, InInstanceData.LastGain, *InAdditionalOperands[i]);
+				TArrayView<float> OutResultView(OutResult->GetData(), SIMDCount);
+
+				Audio::ArrayFade(OutResultView, InInstanceData.LastGain, *InAdditionalOperands[i]);
 
 				for (int32 j = SIMDCount; j < OutResult->Num(); ++j)
 				{
diff --git a/Engine/Plugins/Runtime/Metasound/Source/MetasoundStandardNodes/Private/MetasoundMidSideNodes.cpp b/Engine/Plugins/Runtime/Metasound/Source/MetasoundStandardNodes/Private/MetasoundMidSideNodes.cpp
--- a/Engine/Plugins/Runtime/Metasound/Source/MetasoundStandardNodes/Private/MetasoundMidSideNodes.cpp
+++ b/Engine/Plugins/Runtime/Metasound/Source/MetasoundStandardNodes/Private/MetasoundMidSideNodes.cpp
@@ -187,8 +187,8 @@ namespace Metasound
 			}
 			else
 			{
-				Audio::FadeBufferFast(*AudioOutputMid, PrevMidScale, MidScale);
-				Audio::FadeBufferFast(*AudioOutputSide, PrevSideScale, SideScale);
+				Audio::ArrayFade(*AudioOutputMid, PrevMidScale, MidScale);
+				Audio::ArrayFade(*AudioOutputSide, PrevSideScale, SideScale);
 				
 				PrevMidScale = MidScale;
 				PrevSideScale = SideScale;
@@ -417,8 +417,8 @@ namespace Metasound
 			}
 			else
 			{
-				Audio::FadeBufferFast(MidInBuffer, PrevMidScale, MidScale);
-				Audio::FadeBufferFast(SideInBuffer, PrevSideScale, SideScale);
+				Audio::ArrayFade(MidInBuffer, PrevMidScale, MidScale);
+				Audio::ArrayFade(SideInBuffer, PrevSideScale, SideScale);
 
 				PrevMidScale = MidScale;
 				PrevSideScale = SideScale;
diff --git a/Engine/Plugins/Runtime/Metasound/Source/MetasoundStandardNodes/Private/MetasoundMixerNode.cpp b/Engine/Plugins/Runtime/Metasound/Source/MetasoundStandardNodes/Private/MetasoundMixerNode.cpp
--- a/Engine/Plugins/Runtime/Metasound/Source/MetasoundStandardNodes/Private/MetasoundMixerNode.cpp
+++ b/Engine/Plugins/Runtime/Metasound/Source/MetasoundStandardNodes/Private/MetasoundMixerNode.cpp
@@ -201,12 +201,12 @@ namespace Metasound
 				for (uint32 ChanIndex = 0; ChanIndex < NumChannels; ++ChanIndex)
 				{
 					// Outputs[Chan] += Gains[i] * Inputs[i][Chan]
-					const float* InputPtr = Inputs[InputIndex * NumChannels + ChanIndex]->GetData();
+					TArrayView<const float> InputView(Inputs[InputIndex * NumChannels + ChanIndex]->GetData(), Settings.GetNumFramesPerBlock());
+					TArrayView<float> OutputView(Outputs[ChanIndex]->GetData(), Settings.GetNumFramesPerBlock());
 					const float NextGain = *Gains[InputIndex];
 					const float PrevGain = PrevGains[InputIndex];
-					float* OutputPtr = Outputs[ChanIndex]->GetData();
 
-					Audio::MixInBufferFast(InputPtr,  OutputPtr, Settings.GetNumFramesPerBlock() , PrevGain, NextGain);
+					Audio::ArrayMixIn(InputView, OutputView, PrevGain, NextGain);
 
 					PrevGains[InputIndex] = NextGain;
 				}
diff --git a/Engine/Plugins/Runtime/Metasound/Source/MetasoundStandardNodes/Private/MetasoundStereoPannerNode.cpp b/Engine/Plugins/Runtime/Metasound/Source/MetasoundStandardNodes/Private/MetasoundStereoPannerNode.cpp
--- a/Engine/Plugins/Runtime/Metasound/Source/MetasoundStandardNodes/Private/MetasoundStereoPannerNode.cpp
+++ b/Engine/Plugins/Runtime/Metasound/Source/MetasoundStandardNodes/Private/MetasoundStereoPannerNode.cpp
@@ -147,10 +147,14 @@ namespace Metasound
 		float* OutputLeftBufferPtr = AudioLeftOutput->GetData();
 		float* OutputRightBufferPtr = AudioRightOutput->GetData();
 
+		TArrayView<const float> InputBufferView(AudioInput->GetData(), InputSampleCount);
+		TArrayView<float> OutputLeftBufferView(AudioLeftOutput->GetData(), InputSampleCount);
+		TArrayView<float> OutputRightBufferView(AudioRightOutput->GetData(), InputSampleCount);
+
 		if (FMath::IsNearlyEqual(PrevPanningAmount, CurrentPanningAmount))
 		{
-			Audio::BufferMultiplyByConstant(InputBufferPtr, PrevLeftPan, OutputLeftBufferPtr, InputSampleCount);
-			Audio::BufferMultiplyByConstant(InputBufferPtr, PrevRightPan, OutputRightBufferPtr, InputSampleCount);
+			Audio::ArrayMultiplyByConstant(InputBufferView, PrevLeftPan, OutputLeftBufferView);
+			Audio::ArrayMultiplyByConstant(InputBufferView, PrevRightPan, OutputRightBufferView);
 		}
 		else 
 		{
@@ -164,8 +168,8 @@ namespace Metasound
 			FMemory::Memcpy(OutputRightBufferPtr, InputBufferPtr, InputSampleCount * sizeof(float));
 
 			// Do a fast fade on the buffers from the prev left/right gains to current left/right gains
-			Audio::FadeBufferFast(OutputLeftBufferPtr, InputSampleCount, PrevLeftPan, CurrentLeftPan);
-			Audio::FadeBufferFast(OutputRightBufferPtr, InputSampleCount, PrevRightPan, CurrentRightPan);
+			Audio::ArrayFade(OutputLeftBufferView, PrevLeftPan, CurrentLeftPan);
+			Audio::ArrayFade(OutputRightBufferView, PrevRightPan, CurrentRightPan);
 
 			// lerp through the buffer to the target panning amount
 			PrevPanningAmount = *PanningAmount;
diff --git a/Engine/Plugins/Runtime/Oculus/OculusAudio/Source/OculusAudio/Private/OculusAmbisonicSpatializer.cpp b/Engine/Plugins/Runtime/Oculus/OculusAudio/Source/OculusAudio/Private/OculusAmbisonicSpatializer.cpp
--- a/Engine/Plugins/Runtime/Oculus/OculusAudio/Source/OculusAudio/Private/OculusAmbisonicSpatializer.cpp
+++ b/Engine/Plugins/Runtime/Oculus/OculusAudio/Source/OculusAudio/Private/OculusAmbisonicSpatializer.cpp
@@ -180,7 +180,7 @@ public:
 			OVR_AUDIO_CHECK(Result, "Failed to spatialize mono source interleaved");
 
 			// sum the interleaved output into the output bed.
-			Audio::MixInBufferFast(ScratchOutputBuffer, OutputBuffer.AudioBuffer);
+			Audio::ArrayMixIn(ScratchOutputBuffer, OutputBuffer.AudioBuffer);
 		}
 	}
 };
@@ -222,7 +222,7 @@ public:
 			else if (InputData.PositionalData.NumChannels == InputBuffer.NumChannels)
 			{
 				// If the number of output channels is the same as the input channels, mix it in directly.
-				Audio::MixInBufferFast(InputBuffer.AudioBuffer, OutputData.AudioBuffer);
+				Audio::ArrayMixIn(InputBuffer.AudioBuffer, OutputData.AudioBuffer);
 			}
 			else
 			{
@@ -316,7 +316,7 @@ public:
 			check(InputBuffer.NumChannels == OutputBuffer.NumChannels);
 
 
-			Audio::MixInBufferFast(InputBuffer.AudioBuffer, OutputBuffer.AudioBuffer, InputData.SendLevel);
+			Audio::ArrayMixIn(InputBuffer.AudioBuffer, OutputBuffer.AudioBuffer, InputData.SendLevel);
 		}
 	}
 };
diff --git a/Engine/Plugins/Runtime/SoundFields/Source/SoundFields/Private/SoundFields.cpp b/Engine/Plugins/Runtime/SoundFields/Source/SoundFields/Private/SoundFields.cpp
--- a/Engine/Plugins/Runtime/SoundFields/Source/SoundFields/Private/SoundFields.cpp
+++ b/Engine/Plugins/Runtime/SoundFields/Source/SoundFields/Private/SoundFields.cpp
@@ -93,7 +93,7 @@ public:
 			Decoder.DecodeAudioDirectlyToDeviceOutputPositions(*AudioToDecode, OutputPositions, DecoderOutputBuffer);
 		}
 
-		Audio::MixInBufferFast(DecoderOutputBuffer, OutputAudio);
+		Audio::ArrayMixIn(DecoderOutputBuffer, OutputAudio);
 	}
 };
 
@@ -211,7 +211,7 @@ public:
 		const FQuat AmountToRotateBy = InAudio.Rotation.Inverse();
 		const FQuat PreviousAmountRotatedBy = InAudio.PreviousRotation.Inverse();
 		FSoundFieldDecoder::RotateFirstOrderAmbisonicsBed(InAudio, RotatedAudio, AmountToRotateBy, PreviousAmountRotatedBy);
-		Audio::MixInBufferFast(RotatedAudio.AudioBuffer, OutAudio.AudioBuffer, InputData.SendLevel);
+		Audio::ArrayMixIn(RotatedAudio.AudioBuffer, OutAudio.AudioBuffer, InputData.SendLevel);
 	}
 };
 
diff --git a/Engine/Plugins/Runtime/Synthesis/Source/Synthesis/Private/ConvolutionReverb.cpp b/Engine/Plugins/Runtime/Synthesis/Source/Synthesis/Private/ConvolutionReverb.cpp
--- a/Engine/Plugins/Runtime/Synthesis/Source/Synthesis/Private/ConvolutionReverb.cpp
+++ b/Engine/Plugins/Runtime/Synthesis/Source/Synthesis/Private/ConvolutionReverb.cpp
@@ -606,7 +606,7 @@ namespace Audio
 		InterleaveBuffer(OutputAudio, OutputDeinterleaveBuffers, InNumOutputChannels);
 
 		// Apply final gain
-		MultiplyBufferByConstantInPlace(OutputAudio, OutputGain);
+		ArrayMultiplyByConstantInPlace(OutputAudio, OutputGain);
 	}
 
 	int32 FConvolutionReverb::GetNumInputChannels() const
diff --git a/Engine/Plugins/Runtime/Synthesis/Source/Synthesis/Private/SourceEffectConvolutionReverb.cpp b/Engine/Plugins/Runtime/Synthesis/Source/Synthesis/Private/SourceEffectConvolutionReverb.cpp
--- a/Engine/Plugins/Runtime/Synthesis/Source/Synthesis/Private/SourceEffectConvolutionReverb.cpp
+++ b/Engine/Plugins/Runtime/Synthesis/Source/Synthesis/Private/SourceEffectConvolutionReverb.cpp
@@ -81,14 +81,17 @@ void FSourceEffectConvolutionReverb::ProcessAudio(const FSoundEffectSourceInputD
 
 	Reverb->ProcessAudio(NumChannels, InData.InputSourceEffectBufferPtr, NumChannels, OutAudioBufferData, NumFrames);
 
+	TArrayView<float> OutAudioBufferView(OutAudioBufferData, InData.NumSamples);
+	TArrayView<const float> InputSourceEffectBufferView(InData.InputSourceEffectBufferPtr, InData.NumSamples);
+
 	// Process Wet/Dry mix
 	if (FMath::IsNearlyEqual(WetVolume, 1.f) == false)
 	{
-		Audio::MultiplyBufferByConstantInPlace(OutAudioBufferData, InData.NumSamples, WetVolume);
+		Audio::ArrayMultiplyByConstantInPlace(OutAudioBufferView, WetVolume);
 	}
 	if (DryVolume > KINDA_SMALL_NUMBER)
 	{
-		Audio::MixInBufferFast(InData.InputSourceEffectBufferPtr, OutAudioBufferData, InData.NumSamples, DryVolume);
+		Audio::ArrayMixIn(InputSourceEffectBufferView, OutAudioBufferView, DryVolume);
 	}
 }
 
diff --git a/Engine/Plugins/Runtime/Synthesis/Source/Synthesis/Private/SourceEffectFilter.cpp b/Engine/Plugins/Runtime/Synthesis/Source/Synthesis/Private/SourceEffectFilter.cpp
--- a/Engine/Plugins/Runtime/Synthesis/Source/Synthesis/Private/SourceEffectFilter.cpp
+++ b/Engine/Plugins/Runtime/Synthesis/Source/Synthesis/Private/SourceEffectFilter.cpp
@@ -163,7 +163,7 @@ void FSourceEffectFilter::ProcessAudio(const FSoundEffectSourceInputData& InData
 			ScratchEnvFollowerBuffer.AddUninitialized(InData.NumSamples);
 
 			Audio::BufferSum2ChannelToMonoFast(ScratchModBuffer, ScratchEnvFollowerBuffer);
-			Audio::MultiplyBufferByConstantInPlace(ScratchEnvFollowerBuffer, 0.5f);
+			Audio::ArrayMultiplyByConstantInPlace(ScratchEnvFollowerBuffer, 0.5f);
 
 			Mod.AudioBusEnvelopeFollower.ProcessAudio(ScratchEnvFollowerBuffer.GetData(), InData.NumSamples);
 		}
diff --git a/Engine/Plugins/Runtime/Synthesis/Source/Synthesis/Private/SourceEffectMotionFilter.cpp b/Engine/Plugins/Runtime/Synthesis/Source/Synthesis/Private/SourceEffectMotionFilter.cpp
--- a/Engine/Plugins/Runtime/Synthesis/Source/Synthesis/Private/SourceEffectMotionFilter.cpp
+++ b/Engine/Plugins/Runtime/Synthesis/Source/Synthesis/Private/SourceEffectMotionFilter.cpp
@@ -230,6 +230,10 @@ void FSourceEffectMotionFilter::ProcessAudio(const FSoundEffectSourceInputData&
 	FilterAMixScale = FMath::Clamp(FilterAMixScale, 0.0f, 2.0f);
 	FilterBMixScale = FMath::Clamp(FilterBMixScale, 0.0f, 2.0f);
 
+	TArrayView<float> ScratchBufferAView(ScratchBufferA.GetData(), InData.NumSamples);
+	TArrayView<float> ScratchBufferBView(ScratchBufferB.GetData(), InData.NumSamples);
+	TArrayView<float> OutAudioBufferDataView(OutAudioBufferData, InData.NumSamples);
+	TArrayView<float> InputSourceEffectBufferView(InData.InputSourceEffectBufferPtr, InData.NumSamples);
 
 	// Mix and Processing chain depends on Topology
 	if (Topology == ESourceEffectMotionFilterTopology::ParallelMode)
@@ -239,9 +243,9 @@ void FSourceEffectMotionFilter::ProcessAudio(const FSoundEffectSourceInputData&
 		MotionFilterB.CurrentFilter->ProcessAudio(InData.InputSourceEffectBufferPtr, InData.NumSamples, ScratchBufferB.GetData());
 
 		// Mix Dry + Filter A + Filter B
-		Audio::MixInBufferFast(InData.InputSourceEffectBufferPtr, OutAudioBufferData, InData.NumSamples, DryVolumeScalar);
-		Audio::MixInBufferFast(ScratchBufferA.GetData(), OutAudioBufferData, InData.NumSamples, FilterAMixScale);
-		Audio::MixInBufferFast(ScratchBufferB.GetData(), OutAudioBufferData, InData.NumSamples, FilterBMixScale);
+		Audio::ArrayMixIn(InputSourceEffectBufferView, OutAudioBufferDataView, DryVolumeScalar);
+		Audio::ArrayMixIn(ScratchBufferAView, OutAudioBufferDataView, FilterAMixScale);
+		Audio::ArrayMixIn(ScratchBufferBView, OutAudioBufferDataView, FilterBMixScale);
 
 	}
 	else if (Topology == ESourceEffectMotionFilterTopology::SerialMode)
@@ -250,15 +254,15 @@ void FSourceEffectMotionFilter::ProcessAudio(const FSoundEffectSourceInputData&
 		MotionFilterA.CurrentFilter->ProcessAudio(InData.InputSourceEffectBufferPtr, InData.NumSamples, ScratchBufferA.GetData());
 
 		// Mix Dry + Filter A
-		Audio::MultiplyBufferByConstantInPlace(ScratchBufferA, FilterAMixScale);
-		Audio::MixInBufferFast(InData.InputSourceEffectBufferPtr, ScratchBufferA.GetData(), InData.NumSamples, DryVolumeScalar);
+		Audio::ArrayMultiplyByConstantInPlace(ScratchBufferAView, FilterAMixScale);
+		Audio::ArrayMixIn(InputSourceEffectBufferView, ScratchBufferAView, DryVolumeScalar);
 
 		// Process Filter A Wet/Dry Mix into Filter B
 		MotionFilterB.CurrentFilter->ProcessAudio(ScratchBufferA.GetData(), InData.NumSamples, ScratchBufferB.GetData());
 
 		// Mix Dry + Filter B Output
-		Audio::MixInBufferFast(InData.InputSourceEffectBufferPtr, OutAudioBufferData, InData.NumSamples, DryVolumeScalar);
-		Audio::MixInBufferFast(ScratchBufferB.GetData(), OutAudioBufferData, InData.NumSamples, FilterBMixScale);
+		Audio::ArrayMixIn(InputSourceEffectBufferView, OutAudioBufferDataView, DryVolumeScalar);
+		Audio::ArrayMixIn(ScratchBufferBView, OutAudioBufferDataView, FilterBMixScale);
 	}
 	else
 	{
diff --git a/Engine/Plugins/Runtime/Synthesis/Source/Synthesis/Private/SubmixEffectConvolutionReverb.cpp b/Engine/Plugins/Runtime/Synthesis/Source/Synthesis/Private/SubmixEffectConvolutionReverb.cpp
--- a/Engine/Plugins/Runtime/Synthesis/Source/Synthesis/Private/SubmixEffectConvolutionReverb.cpp
+++ b/Engine/Plugins/Runtime/Synthesis/Source/Synthesis/Private/SubmixEffectConvolutionReverb.cpp
@@ -97,11 +97,11 @@ void FSubmixEffectConvolutionReverb::OnProcessAudio(const FSoundEffectSubmixInpu
 	// Process Wet/Dry mix
 	if (FMath::IsNearlyEqual(WetVolume, 1.f) == false)
 	{
-		Audio::MultiplyBufferByConstantInPlace(*OutData.AudioBuffer, WetVolume);
+		Audio::ArrayMultiplyByConstantInPlace(*OutData.AudioBuffer, WetVolume);
 	}
 	if (DryVolume > KINDA_SMALL_NUMBER)
 	{
-		Audio::MixInBufferFast(*InData.AudioBuffer, *OutData.AudioBuffer, DryVolume);
+		Audio::ArrayMixIn(*InData.AudioBuffer, *OutData.AudioBuffer, DryVolume);
 	}
 }
 
diff --git a/Engine/Plugins/Runtime/Synthesis/Source/Synthesis/Private/SubmixEffectMultiBandCompressor.cpp b/Engine/Plugins/Runtime/Synthesis/Source/Synthesis/Private/SubmixEffectMultiBandCompressor.cpp
--- a/Engine/Plugins/Runtime/Synthesis/Source/Synthesis/Private/SubmixEffectMultiBandCompressor.cpp
+++ b/Engine/Plugins/Runtime/Synthesis/Source/Synthesis/Private/SubmixEffectMultiBandCompressor.cpp
@@ -433,10 +433,13 @@ void FSubmixEffectMultibandCompressor::OnProcessAudio(const FSoundEffectSubmixIn
 			}
 		}
 
+		TArrayView<const float> ScratchBufferView(ScratchBuffer.GetData(), BlockSize);
+		TArrayView<float> OutPtrView(OutPtr, BlockSize);
+
 		for (int32 Band = 0; Band < DynamicsProcessors.Num(); ++Band)
 		{
 			DynamicsProcessors[Band].ProcessAudio(MultiBandBuffer[Band], BlockSize, ScratchBuffer.GetData(), bUseKey ? KeyMultiBandBuffer[Band] : nullptr);
-			Audio::MixInBufferFast(ScratchBuffer.GetData(), OutPtr, BlockSize);
+			Audio::ArrayMixIn(ScratchBufferView, OutPtrView);
 		}
 	}
 }
diff --git a/Engine/Plugins/Runtime/Synthesis/Source/Synthesis/Private/SynthComponentMonoWaveTable.cpp b/Engine/Plugins/Runtime/Synthesis/Source/Synthesis/Private/SynthComponentMonoWaveTable.cpp
--- a/Engine/Plugins/Runtime/Synthesis/Source/Synthesis/Private/SynthComponentMonoWaveTable.cpp
+++ b/Engine/Plugins/Runtime/Synthesis/Source/Synthesis/Private/SynthComponentMonoWaveTable.cpp
@@ -418,10 +418,12 @@ void USynthComponentMonoWaveTable::RefreshWaveTable(int32 Index)
 	MaxValue += DCShift;
 	MinValue += DCShift;
 
+	TArrayView<float> DestinationTableView(DestinationTable.GetData(), TableRes);
+
 	// Remove DC offset
 	if (!FMath::IsNearlyZero(DCShift))
 	{
-		Audio::AddConstantToBufferInplace(DestinationTable.GetData(), TableRes, DCShift);
+		Audio::ArrayAddConstantInplace(DestinationTableView, DCShift);
 	}
 
 	if (CachedPreset->bNormalizeWaveTables)
@@ -429,7 +431,7 @@ void USynthComponentMonoWaveTable::RefreshWaveTable(int32 Index)
 		const float MaxAbsValue = FMath::Max(FMath::Abs(MaxValue), FMath::Abs(MinValue));
 		const float NormalizationScalar = 1.0f / MaxAbsValue;
 
-		Audio::MultiplyBufferByConstantInPlace(DestinationTable.GetData(), TableRes, NormalizationScalar);
+		Audio::ArrayMultiplyByConstantInPlace(DestinationTableView, NormalizationScalar);
 	}
 	else // clip the values since we aren't normalizing
 	{
diff --git a/Engine/Source/Runtime/AudioMixer/Private/AudioMixerBus.cpp b/Engine/Source/Runtime/AudioMixer/Private/AudioMixerBus.cpp
--- a/Engine/Source/Runtime/AudioMixer/Private/AudioMixerBus.cpp
+++ b/Engine/Source/Runtime/AudioMixer/Private/AudioMixerBus.cpp
@@ -137,7 +137,9 @@ namespace Audio
 					}
 					else
 					{
-						MixInBufferFast(SourceBufferPtr, BusDataBufferPtr, NumOutputFrames * NumChannels, AudioBusSend.SendLevel);
+						TArrayView<const float> SourceBufferView(SourceBufferPtr, NumOutputFrames * NumChannels);
+						TArrayView<float> BusDataBufferView(BusDataBufferPtr, NumOutputFrames * NumChannels);
+						ArrayMixIn(SourceBufferView, BusDataBufferView, AudioBusSend.SendLevel);
 					}
 				}
 			}
diff --git a/Engine/Source/Runtime/AudioMixer/Private/AudioMixerSourceManager.cpp b/Engine/Source/Runtime/AudioMixer/Private/AudioMixerSourceManager.cpp
--- a/Engine/Source/Runtime/AudioMixer/Private/AudioMixerSourceManager.cpp
+++ b/Engine/Source/Runtime/AudioMixer/Private/AudioMixerSourceManager.cpp
@@ -2074,8 +2074,8 @@ namespace Audio
 			return;
 		}
 
-		float* PostDistanceAttenBufferPtr = SourceInfo.SourceBuffer.GetData();
-		Audio::FadeBufferFast(PostDistanceAttenBufferPtr, SourceInfo.SourceBuffer.Num(), SourceInfo.DistanceAttenuationSourceStart, SourceInfo.DistanceAttenuationSourceDestination);
+		TArrayView<float> PostDistanceAttenBufferView(SourceInfo.SourceBuffer.GetData(), SourceInfo.SourceBuffer.Num());
+		Audio::ArrayFade(PostDistanceAttenBufferView, SourceInfo.DistanceAttenuationSourceStart, SourceInfo.DistanceAttenuationSourceDestination);
 		SourceInfo.DistanceAttenuationSourceStart = SourceInfo.DistanceAttenuationSourceDestination;
 	}
 
@@ -2148,6 +2148,10 @@ namespace Audio
 			}
 		}
 
+		TArrayView<const float> ReverbPluginOutputBufferView(InSourceSubmixOutputBuffer.GetReverbPluginOutputData(), NumSamples);
+		TArrayView<const float> AudioPluginOutputDataView(SourceInfo.AudioPluginOutputData.AudioBuffer.GetData(), NumSamples);
+		TArrayView<float> PostDistanceAttenBufferView(PostDistanceAttenBufferPtr, NumSamples);
+
 		if (SourceInfo.bUseOcclusionPlugin)
 		{
 			const FSpatializationParams* SourceSpatParams = &SourceInfo.SpatParams;
@@ -2171,10 +2175,7 @@ namespace Audio
 			// Copy the occlusion-processed data back to the source buffer and mix with the reverb plugin output buffer
 			if (bShouldMixInReverb)
 			{
-				const float* ReverbPluginOutputBufferPtr = InSourceSubmixOutputBuffer.GetReverbPluginOutputData();
-				const float* AudioPluginOutputDataPtr = SourceInfo.AudioPluginOutputData.AudioBuffer.GetData();
-
-				Audio::SumBuffers(ReverbPluginOutputBufferPtr, AudioPluginOutputDataPtr, PostDistanceAttenBufferPtr, NumSamples);
+				Audio::ArraySum(ReverbPluginOutputBufferView, AudioPluginOutputDataView, PostDistanceAttenBufferView);
 			}
 			else
 			{
@@ -2183,8 +2184,7 @@ namespace Audio
 		}
 		else if (bShouldMixInReverb)
 		{
-			const float* ReverbPluginOutputBufferPtr = InSourceSubmixOutputBuffer.GetReverbPluginOutputData();
-			Audio::MixInBufferFast(ReverbPluginOutputBufferPtr, PostDistanceAttenBufferPtr, NumSamples);
+			Audio::ArrayMixIn(ReverbPluginOutputBufferView, PostDistanceAttenBufferView);
 		}
 
 		// If the source has HRTF processing enabled, run it through the spatializer
@@ -2294,6 +2294,8 @@ namespace Audio
 			float* PreDistanceAttenBufferPtr = SourceInfo.PreDistanceAttenuationBuffer.GetData();
 			const int32 NumSamples = SourceInfo.PreDistanceAttenuationBuffer.Num();
 
+			TArrayView<float> PreDistanceAttenBufferView(PreDistanceAttenBufferPtr, NumSamples);
+
 			// Update volume fade information if we're stopping
 			if (SourceInfo.bIsStopping)
 			{
@@ -2328,7 +2330,8 @@ namespace Audio
 					VolumeDestination *= ModVolumeEnd;
 				}
 
-				Audio::FadeBufferFast(PreDistanceAttenBufferPtr, NumFadeSamples, VolumeStart, VolumeDestination);
+				TArrayView<float> PreDistanceAttenBufferFadeSamplesView(PreDistanceAttenBufferPtr, NumFadeSamples);
+				Audio::ArrayFade(PreDistanceAttenBufferFadeSamplesView, VolumeStart, VolumeDestination);
 
 				// Zero the rest of the buffer
 				if (NumFadeFrames < NumOutputFrames)
@@ -2362,7 +2365,8 @@ namespace Audio
 					}
 					VolumeDestination *= ModVolumeEnd;
 				}
-				Audio::FadeBufferFast(PreDistanceAttenBufferPtr, NumSamples, VolumeStart, VolumeDestination);
+
+				Audio::ArrayFade(PreDistanceAttenBufferView, VolumeStart, VolumeDestination);
 			}
 			SourceInfo.VolumeSourceStart = SourceInfo.VolumeSourceDestination;
 
@@ -2414,7 +2418,7 @@ namespace Audio
 			if (!DisableEnvelopeFollowingCvar)
 			{
 				// Compute the source envelope using pre-distance attenuation buffer
-				float AverageSampleValue = Audio::BufferGetAverageAbsValue(PreDistanceAttenBufferPtr, NumSamples);
+				float AverageSampleValue = Audio::ArrayGetAverageAbsValue(PreDistanceAttenBufferView);
 				SourceInfo.SourceEnvelopeValue = SourceInfo.SourceEnvelopeFollower.ProcessSample(AverageSampleValue);
 				SourceInfo.SourceEnvelopeValue = FMath::Clamp(SourceInfo.SourceEnvelopeValue, 0.f, 1.f);
 
diff --git a/Engine/Source/Runtime/AudioMixer/Private/AudioMixerSourceOutputBuffer.cpp b/Engine/Source/Runtime/AudioMixer/Private/AudioMixerSourceOutputBuffer.cpp
--- a/Engine/Source/Runtime/AudioMixer/Private/AudioMixerSourceOutputBuffer.cpp
+++ b/Engine/Source/Runtime/AudioMixer/Private/AudioMixerSourceOutputBuffer.cpp
@@ -492,11 +492,11 @@ namespace Audio
 	{
 		if (InSubmixSendStage == EMixerSourceSubmixSendStage::PostDistanceAttenuation)
 		{
-			Audio::MixInBufferFast(PostAttenuationOutputBuffer, OutMixedBuffer, InSendLevel);
+			Audio::ArrayMixIn(PostAttenuationOutputBuffer, OutMixedBuffer, InSendLevel);
 		}
 		else
 		{
-			Audio::MixInBufferFast(PreAttenuationOutputBuffer, OutMixedBuffer, InSendLevel);
+			Audio::ArrayMixIn(PreAttenuationOutputBuffer, OutMixedBuffer, InSendLevel);
 		}
 	}
 
diff --git a/Engine/Source/Runtime/AudioMixer/Private/AudioMixerSubmix.cpp b/Engine/Source/Runtime/AudioMixer/Private/AudioMixerSubmix.cpp
--- a/Engine/Source/Runtime/AudioMixer/Private/AudioMixerSubmix.cpp
+++ b/Engine/Source/Runtime/AudioMixer/Private/AudioMixerSubmix.cpp
@@ -1342,7 +1342,7 @@ namespace Audio
 					float EndFadeVolume = FadeInfo.FadeVolume.GetValue();
 
 					// Mix this effect chain with other effect chains.
-					MixInBufferFast(EffectChainOutputBuffer, SubmixChainMixBuffer, StartFadeVolume, EndFadeVolume);
+					ArrayMixIn(EffectChainOutputBuffer, SubmixChainMixBuffer, StartFadeVolume, EndFadeVolume);
 				}
 
 				// If we processed any effects, write over the old input buffer vs mixing into it. This is basically the "wet channel" audio in a submix.
@@ -1369,11 +1369,11 @@ namespace Audio
 				{
 					if (FMath::IsNearlyEqual(TargetWetLevel, CurrentWetLevel))
 					{
-						MultiplyBufferByConstantInPlace(InputBuffer, TargetWetLevel);
+						ArrayMultiplyByConstantInPlace(InputBuffer, TargetWetLevel);
 					}
 					else
 					{
-						FadeBufferFast(InputBuffer, CurrentWetLevel, TargetWetLevel);
+						ArrayFade(InputBuffer, CurrentWetLevel, TargetWetLevel);
 						CurrentWetLevel = TargetWetLevel;
 					}
 				}
@@ -1386,15 +1386,15 @@ namespace Audio
 			// If we've already set the volume, only need to multiply by constant
 			if (FMath::IsNearlyEqual(TargetDryLevel, CurrentDryLevel))
 			{
-				MultiplyBufferByConstantInPlace(DryChannelBuffer, TargetDryLevel);
+				ArrayMultiplyByConstantInPlace(DryChannelBuffer, TargetDryLevel);
 			}
 			else
 			{
 				// To avoid popping, we do a fade on the buffer to the target volume
-				FadeBufferFast(DryChannelBuffer, CurrentDryLevel, TargetDryLevel);
+				ArrayFade(DryChannelBuffer, CurrentDryLevel, TargetDryLevel);
 				CurrentDryLevel = TargetDryLevel;
 			}
-			MixInBufferFast(DryChannelBuffer, InputBuffer);
+			ArrayMixIn(DryChannelBuffer, InputBuffer);
 		}
 
 		// If we're muted, memzero the buffer. Note we are still doing all the work to maintain buffer state between mutings.
@@ -1475,12 +1475,12 @@ namespace Audio
 			// If we've already set the output volume, only need to multiply by constant
 			if (FMath::IsNearlyEqual(TargetOutputVolume, CurrentOutputVolume))
 			{
-				Audio::MultiplyBufferByConstantInPlace(InputBuffer, TargetOutputVolume);
+				Audio::ArrayMultiplyByConstantInPlace(InputBuffer, TargetOutputVolume);
 			}
 			else
 			{
 				// To avoid popping, we do a fade on the buffer to the target volume
-				Audio::FadeBufferFast(InputBuffer, CurrentOutputVolume, TargetOutputVolume);
+				Audio::ArrayFade(InputBuffer, CurrentOutputVolume, TargetOutputVolume);
 				CurrentOutputVolume = TargetOutputVolume;
 			}
 		}
@@ -1488,7 +1488,7 @@ namespace Audio
 		SendAudioToSubmixBufferListeners(InputBuffer);
 
 		// Mix the audio buffer of this submix with the audio buffer of the output buffer (i.e. with other submixes)
-		Audio::MixInBufferFast(InputBuffer, OutAudioBuffer);
+		Audio::ArrayMixIn(InputBuffer, OutAudioBuffer);
 
 		// Once we've finished rendering submix audio, check if the output buffer is silent if we are auto-disabling
 		if (bAutoDisable)
@@ -1591,7 +1591,7 @@ namespace Audio
 				const float DryLevel = SubmixEffect->GetDryLevel();
 				if (DryLevel > 0.0f)
 				{
-					MixInBufferFast(ScratchBuffer, *OutputData.AudioBuffer, DryLevel);
+					ArrayMixIn(ScratchBuffer, *OutputData.AudioBuffer, DryLevel);
 				}
 				// Copy the output to the input
 				FMemory::Memcpy((void*)ScratchBuffer.GetData(), (void*)OutputData.AudioBuffer->GetData(), sizeof(float) * NumSamples);
diff --git a/Engine/Source/Runtime/AudioMixer/Private/Effects/AudioMixerSubmixEffectReverb.cpp b/Engine/Source/Runtime/AudioMixer/Private/Effects/AudioMixerSubmixEffectReverb.cpp
--- a/Engine/Source/Runtime/AudioMixer/Private/Effects/AudioMixerSubmixEffectReverb.cpp
+++ b/Engine/Source/Runtime/AudioMixer/Private/Effects/AudioMixerSubmixEffectReverb.cpp
@@ -158,7 +158,7 @@ void FSubmixEffectReverb::OnProcessAudio(const FSoundEffectSubmixInputData& InDa
 	{
 		// Wet level is applied to input audio to preserve reverb tail when changing wet level
 		WetInputBuffer.AddZeroed(InData.AudioBuffer->Num());
-		Audio::MixInBufferFast(*InData.AudioBuffer, WetInputBuffer, LastWet, CurrentWetDry.WetLevel);
+		Audio::ArrayMixIn(*InData.AudioBuffer, WetInputBuffer, LastWet, CurrentWetDry.WetLevel);
 	}
 
 	PlateReverb->ProcessAudio(WetInputBuffer, InData.NumChannels, *OutData.AudioBuffer, OutData.NumChannels);
diff --git a/Engine/Source/Runtime/AudioMixerCore/Private/AudioMixer.cpp b/Engine/Source/Runtime/AudioMixerCore/Private/AudioMixer.cpp
--- a/Engine/Source/Runtime/AudioMixerCore/Private/AudioMixer.cpp
+++ b/Engine/Source/Runtime/AudioMixerCore/Private/AudioMixer.cpp
@@ -239,9 +239,9 @@ namespace Audio
 		{
 			if (!FMath::IsNearlyEqual(LinearGainScalarForFinalOututCVar, 1.0f))
 			{
-				MultiplyBufferByConstantInPlace(RenderBuffer, LinearGainScalarForFinalOututCVar);
+				ArrayMultiplyByConstantInPlace(RenderBuffer, LinearGainScalarForFinalOututCVar);
 			}
-			BufferRangeClampFast(RenderBuffer, -1.0f, 1.0f);
+			ArrayRangeClamp(RenderBuffer, -1.0f, 1.0f);
 
 			// No conversion is needed, so we push the RenderBuffer directly to the circular queue.
 			CircularBuffer.Push(reinterpret_cast<const uint8*>(RenderBuffer.GetData()), RenderBuffer.Num() * sizeof(float));
@@ -255,8 +255,8 @@ namespace Audio
 			check(FormattedBuffer.Num() / GetSizeForDataFormat(DataFormat) == RenderBuffer.Num());			
 
 			const float ConversionScalar = LinearGainScalarForFinalOututCVar * 32767.0f;
-			MultiplyBufferByConstantInPlace(RenderBuffer, ConversionScalar);
-			BufferRangeClampFast(RenderBuffer, -32767.0f, 32767.0f);
+			ArrayMultiplyByConstantInPlace(RenderBuffer, ConversionScalar);
+			ArrayRangeClamp(RenderBuffer, -32767.0f, 32767.0f);
 
 			for (int32 i = 0; i < NumSamples; ++i)
 			{
diff --git a/Engine/Source/Runtime/Engine/Private/SoundModulationDestination.cpp b/Engine/Source/Runtime/Engine/Private/SoundModulationDestination.cpp
--- a/Engine/Source/Runtime/Engine/Private/SoundModulationDestination.cpp
+++ b/Engine/Source/Runtime/Engine/Private/SoundModulationDestination.cpp
@@ -170,12 +170,12 @@ namespace Audio
 			{
 				if (FMath::IsNearlyEqual(LastTarget, ValueTarget))
 				{
-					BufferSetToConstantInplace(OutputBuffer, ValueTarget);
+					ArraySetToConstantInplace(OutputBuffer, ValueTarget);
 				}
 				else
 				{
-					BufferSetToConstantInplace(OutputBuffer, 1.0f);
-					FadeBufferFast(OutputBuffer, LastTarget, ValueTarget);
+					ArraySetToConstantInplace(OutputBuffer, 1.0f);
+					ArrayFade(OutputBuffer, LastTarget, ValueTarget);
 				}
 			}
 			else
diff --git a/Engine/Source/Runtime/SignalProcessing/Private/AllPassFractionalDelay.cpp b/Engine/Source/Runtime/SignalProcessing/Private/AllPassFractionalDelay.cpp
--- a/Engine/Source/Runtime/SignalProcessing/Private/AllPassFractionalDelay.cpp
+++ b/Engine/Source/Runtime/SignalProcessing/Private/AllPassFractionalDelay.cpp
@@ -153,7 +153,8 @@ void FAllPassFractionalDelay::ProcessAudioBlock(const float* InSamples, const fl
 			OutSamples[i] = DelayData[DelayPos - 1] + CoefficientsData[i] * DelayData[DelayPos] - CoefficientsData[i] * OutSamples[i - 1];
 		}
 
-		BufferUnderflowClampFast(OutSamples, InNum);
+		TArrayView<float> OutSamplesView(OutSamples, InNum);
+		ArrayUnderflowClamp(OutSamplesView);
 
 		Z1 = OutSamples[InNum - 1];
 	}
diff --git a/Engine/Source/Runtime/SignalProcessing/Private/AudioBufferDistanceAttenuation.cpp b/Engine/Source/Runtime/SignalProcessing/Private/AudioBufferDistanceAttenuation.cpp
--- a/Engine/Source/Runtime/SignalProcessing/Private/AudioBufferDistanceAttenuation.cpp
+++ b/Engine/Source/Runtime/SignalProcessing/Private/AudioBufferDistanceAttenuation.cpp
@@ -77,7 +77,8 @@ namespace Audio
 		// If we're passed in a negative value for InOutAttenuation, that means we don't want to interpolate from that value to target value (i.e. it's the first one).
 		// This prevents a pop when first applying attenuation if a sound is far away.
 		float Gain = InOutAttenuation < 0.0f ? TargetAttenuationLinear : InOutAttenuation;
-		Audio::FadeBufferFast(InOutBuffer.GetData(), NumSamples, Gain, TargetAttenuationLinear);
+		TArrayView<float> InOutBufferView(InOutBuffer.GetData(), NumSamples);
+		Audio::ArrayFade(InOutBufferView, Gain, TargetAttenuationLinear);
 
 		InOutAttenuation = TargetAttenuationLinear;
 	}
diff --git a/Engine/Source/Runtime/SignalProcessing/Private/AudioFFT.cpp b/Engine/Source/Runtime/SignalProcessing/Private/AudioFFT.cpp
--- a/Engine/Source/Runtime/SignalProcessing/Private/AudioFFT.cpp
+++ b/Engine/Source/Runtime/SignalProcessing/Private/AudioFFT.cpp
@@ -204,8 +204,8 @@ namespace Audio
 		{
 			checkSlow(InFreqBuffer.Real.Num() % 4 == 0);
 			
-			MultiplyBufferByConstantInPlace(InFreqBuffer.Real, InReal);
-			MultiplyBufferByConstantInPlace(InFreqBuffer.Imag, InReal);
+			ArrayMultiplyByConstantInPlace(InFreqBuffer.Real, InReal);
+			ArrayMultiplyByConstantInPlace(InFreqBuffer.Imag, InReal);
 		}
 
 		// Given:
@@ -405,11 +405,14 @@ namespace Audio
 
 		void PerformIterativeIFFT(FFTFreqDomainData& InputParams, FFTTimeDomainData& OutputParams)
 		{
+			TArrayView<float> OutImagView(InputParams.OutImag, OutputParams.NumSamples);
+			TArrayView<float> OutputBufferView(OutputParams.Buffer, OutputParams.NumSamples);
+
 			SeperateInPlace(InputParams.OutReal, OutputParams.NumSamples);
 			SeperateInPlace(InputParams.OutImag, OutputParams.NumSamples);
 
 			// IFFT can be done by performing a forward FFT on the complex conjugate of a frequency domain signal:
-			MultiplyBufferByConstantInPlace(InputParams.OutImag, OutputParams.NumSamples, -1.0f);
+			ArrayMultiplyByConstantInPlace(OutImagView, -1.0f);
 
 			// Iterate over and compute butterflies.
 			ComputeButterfliesInPlace(InputParams.OutReal, InputParams.OutImag, OutputParams.NumSamples);
@@ -425,7 +428,7 @@ namespace Audio
 			FMemory::Memcpy(OutputParams.Buffer, InputParams.OutReal, OutputParams.NumSamples * sizeof(float));
 			
 			// Personal note: This is a very important step in an inverse FFT.
-			Audio::MultiplyBufferByConstantInPlace(OutputParams.Buffer, OutputParams.NumSamples, 1.0f / OutputParams.NumSamples);
+			Audio::ArrayMultiplyByConstantInPlace(OutputBufferView, 1.0f / OutputParams.NumSamples);
 #endif
 		}
 
@@ -670,29 +673,11 @@ namespace Audio
 		OutBuffer.Reset(NumSpectrumValues);
 		OutBuffer.AddUninitialized(NumSpectrumValues);
 
-		float* OutBufferData = OutBuffer.GetData();
-
-		const float* RealData = InFrequencyData.OutReal;
-		const float* ImagData = InFrequencyData.OutImag;
+		TArrayView<float> OutBufferDataView(OutBuffer.GetData(), NumSpectrumValues);
+		TArrayView<const float> RealDataView(InFrequencyData.OutReal, NumSpectrumValues);
+		TArrayView<const float> ImagDataView(InFrequencyData.OutImag, NumSpectrumValues);
 
-		if (IsAligned<const float*>(RealData, AUDIO_BUFFER_ALIGNMENT) && IsAligned<const float*>(ImagData, AUDIO_BUFFER_ALIGNMENT) && (NumSpectrumValues > 5))
-		{
-			BufferComplexToPowerFast(RealData, ImagData, OutBufferData, NumSpectrumValues - 1);
-
-			// Buffer operations are on 16 byte boundaries, which excludes the location of the nyquist
-			// frequency in the fft output buffers. Here we explicitly handle data at nyquist freq.
-			const int32 NyquistIndex = NumSpectrumValues - 1;
-			const float NyquistReal = RealData[NyquistIndex];
-			const float NyquistImag = ImagData[NyquistIndex];
-			OutBufferData[NyquistIndex] = NyquistReal * NyquistReal + NyquistImag * NyquistImag;
-		}
-		else
-		{
-			for (int32 i = 0; i < NumSpectrumValues; i++)
-			{
-				OutBufferData[i] = RealData[i] * RealData[i] + ImagData[i] * ImagData[i];
-			}
-		}
+		ArrayComplexToPower(RealDataView, ImagDataView, OutBufferDataView);
 	}
 
 	void ComputePowerSpectrum(const FFTFreqDomainData& InFrequencyData, int32 FFTSize, FAlignedFloatBuffer& OutBuffer)
@@ -717,7 +702,8 @@ namespace Audio
 
 		if (NumSpectrumValues > 1)
 		{
-			MultiplyBufferByConstantInPlace(OutBufferData, NumSpectrumValues - 1, FFTScale);
+			TArrayView<float> OutBufferDataView(OutBufferData, NumSpectrumValues - 1);
+			ArrayMultiplyByConstantInPlace(OutBufferDataView, FFTScale);
 		}
 		
 		OutBufferData[NumSpectrumValues - 1] *= FFTScale;
@@ -751,7 +737,8 @@ namespace Audio
 
 		if (NumSpectrumValues > 1)
 		{
-			MultiplyBufferByConstantInPlace(OutBufferData, NumSpectrumValues - 1, FFTScale);
+			TArrayView<float> OutBufferView(OutBufferData, NumSpectrumValues - 1);
+			ArrayMultiplyByConstantInPlace(OutBufferView, FFTScale);
 		}
 		OutBufferData[NumSpectrumValues - 1] *= FFTScale;
 	}
@@ -941,15 +928,10 @@ namespace Audio
 
 	void FFFTConvolver::SumInCOLABuffer(float* InputAudio, int32 NumSamples)
 	{
-		const float* COLABufferPtr = COLABuffer.GetData();
-		const int32 Remainder = NumSamples % 4;
-		const int32 NumVectorizedSamples = NumSamples - Remainder;
-		MixInBufferFast(COLABufferPtr, InputAudio, NumVectorizedSamples);
-		
-		for (int32 Index = NumVectorizedSamples; Index < NumSamples; Index++)
-		{
-			InputAudio[Index] += COLABufferPtr[Index];
-		}
+		TArrayView<const float> COLABufferView(COLABuffer.GetData(), NumSamples);
+		TArrayView<float> InputAudioView(InputAudio, NumSamples);
+
+		ArrayMixIn(COLABufferView, InputAudioView);
 	}
 
 	void FFFTConvolver::SetCOLABuffer(float* InAudio, int32 NumSamples)
diff --git a/Engine/Source/Runtime/SignalProcessing/Private/BufferOnePoleLPF.cpp b/Engine/Source/Runtime/SignalProcessing/Private/BufferOnePoleLPF.cpp
--- a/Engine/Source/Runtime/SignalProcessing/Private/BufferOnePoleLPF.cpp
+++ b/Engine/Source/Runtime/SignalProcessing/Private/BufferOnePoleLPF.cpp
@@ -63,7 +63,7 @@ void FBufferOnePoleLPF::ProcessAudio(const Audio::FAlignedFloatBuffer& InSamples
 		//OutSampleData[Index] = UnderflowClamp(InSampleData[Index] * A0 + B1 * OutSampleData[DelayIndex]);
 		OutSampleData[Index] = InSampleData[Index] * A0 + B1 * OutSampleData[DelayIndex];
 	}
-	BufferUnderflowClampFast(OutSampleData, InNum);
+	ArrayUnderflowClamp(OutSamples);
 	// Store delay value
 	Z1 = OutSampleData[InNum - 1];
 }
diff --git a/Engine/Source/Runtime/SignalProcessing/Private/BufferVectorOperations.cpp b/Engine/Source/Runtime/SignalProcessing/Private/BufferVectorOperations.cpp
--- a/Engine/Source/Runtime/SignalProcessing/Private/BufferVectorOperations.cpp
+++ b/Engine/Source/Runtime/SignalProcessing/Private/BufferVectorOperations.cpp
@@ -1,703 +1,10 @@
 // Copyright Epic Games, Inc. All Rights Reserved.
 
 #include "DSP/BufferVectorOperations.h"
-
-#define AUDIO_USE_SIMD 1
+#include "DSP/FloatArrayMath.h"
 
 namespace Audio
 {
-	static void RestrictedPtrAliasCheck(const float* RESTRICT Ptr1, const float* RESTRICT Ptr2, uint32 NumFloatsInArray)
-	{
-		checkf(static_cast<uint32>(FMath::Abs(Ptr1 - Ptr2)) >= NumFloatsInArray,
-			TEXT("Using this function as an in-place operation will result in undefined behavior!"));
-	}
-
-	/* Sets a values to zero if value is denormal. Denormal numbers significantly slow down floating point operations. */
-	void BufferUnderflowClampFast(FAlignedFloatBuffer& InOutBuffer)
-	{
-		BufferUnderflowClampFast(InOutBuffer.GetData(), InOutBuffer.Num());
-	}
-	
-	/* Sets a values to zero if value is denormal. Denormal numbers significantly slow down floating point operations. */
-	void BufferUnderflowClampFast(float* RESTRICT InOutBuffer, const int32 InNum)
-	{
-		checkf(InNum >= AUDIO_NUM_FLOATS_PER_VECTOR_REGISTER, TEXT("Buffer must have at least %i elements."), AUDIO_NUM_FLOATS_PER_VECTOR_REGISTER);
-		checkf(0 == (InNum % AUDIO_NUM_FLOATS_PER_VECTOR_REGISTER), TEXT("Buffer length be a multiple of %i."), AUDIO_NUM_FLOATS_PER_VECTOR_REGISTER);
-		checkf(IsAligned<float*>(InOutBuffer, AUDIO_SIMD_BYTE_ALIGNMENT), TEXT("Memory must be aligned to use vector operations."));
-
-		const VectorRegister4Float VFMIN = MakeVectorRegister(FLT_MIN, FLT_MIN, FLT_MIN, FLT_MIN);
-		const VectorRegister4Float VNFMIN = MakeVectorRegister(-FLT_MIN, -FLT_MIN, -FLT_MIN, -FLT_MIN);
-
-		for (int32 i = 0; i < InNum; i += AUDIO_NUM_FLOATS_PER_VECTOR_REGISTER)
-		{
-			VectorRegister4Float VInOut = VectorLoadAligned(&InOutBuffer[i]);
-
-			// Create mask of denormal numbers.
-			VectorRegister4Float Mask = VectorBitwiseAnd(VectorCompareGT(VInOut, VNFMIN), VectorCompareLT(VInOut, VFMIN));
-
-			// Choose between zero or original number based upon mask.
-			VInOut = VectorSelect(Mask, GlobalVectorConstants::FloatZero, VInOut);
-			VectorStoreAligned(VInOut, &InOutBuffer[i]);
-		}
-	}
-
-	/* Clamps values in the buffer to be between InMinValue and InMaxValue */
-	void BufferRangeClampFast(FAlignedFloatBuffer& InOutBuffer, float InMinValue, float InMaxValue)
-	{
-		return BufferRangeClampFast(InOutBuffer.GetData(), InOutBuffer.Num(), InMinValue, InMaxValue);
-	}
-
-	/* Clamps values in the buffer to be between InMinValue and InMaxValue */
-	void BufferRangeClampFast(float* RESTRICT InOutBuffer, const int32 InNum, float InMinValue, float InMaxValue)
-	{
-		checkf(InNum >= AUDIO_NUM_FLOATS_PER_VECTOR_REGISTER, TEXT("Buffer must have at least %i elements."), AUDIO_NUM_FLOATS_PER_VECTOR_REGISTER);
-		checkf(0 == (InNum % AUDIO_NUM_FLOATS_PER_VECTOR_REGISTER), TEXT("Buffer length be a multiple of %i."), AUDIO_NUM_FLOATS_PER_VECTOR_REGISTER);
-		checkf(IsAligned<float*>(InOutBuffer, AUDIO_SIMD_BYTE_ALIGNMENT), TEXT("Memory must be aligned to use vector operations."));
-
-		const VectorRegister4Float VMinVal = MakeVectorRegister(InMinValue, InMinValue, InMinValue, InMinValue);
-		const VectorRegister4Float VMaxVal = MakeVectorRegister(InMaxValue, InMaxValue, InMaxValue, InMaxValue);
-
-		for (int32 i = 0; i < InNum; i += AUDIO_NUM_FLOATS_PER_VECTOR_REGISTER)
-		{
-			VectorRegister4Float VInOut = VectorLoadAligned(&InOutBuffer[i]);
-
-			// Create masks to flag elements outside of range.
-			VectorRegister4Float MinMask = VectorCompareLT(VInOut, VMinVal);
-			VectorRegister4Float MaxMask = VectorCompareGT(VInOut, VMaxVal);
-
-			// Choose between range extremes or original number based on masks.
-			VInOut = VectorSelect(MinMask, VMinVal, VInOut);
-			VInOut = VectorSelect(MaxMask, VMaxVal, VInOut);
-
-			VectorStoreAligned(VInOut, &InOutBuffer[i]);
-		}
-	}
-
-	void BufferMultiplyByConstant(const FAlignedFloatBuffer& InFloatBuffer, float InValue, FAlignedFloatBuffer& OutFloatBuffer)
-	{
-		// Prepare output buffer
-		OutFloatBuffer.Reset();
-		OutFloatBuffer.AddUninitialized(InFloatBuffer.Num());
-
-		check(InFloatBuffer.Num() == OutFloatBuffer.Num());
-
-		const int32 NumSamples = InFloatBuffer.Num();
-
-		// Get ptrs to audio buffers to avoid bounds check in non-shipping builds
-		const float* InBufferPtr = InFloatBuffer.GetData();
-		float* OutBufferPtr = OutFloatBuffer.GetData();
-
-		BufferMultiplyByConstant(InBufferPtr, InValue, OutBufferPtr, NumSamples);
-	}
-
-	void BufferMultiplyByConstant(const float* RESTRICT InFloatBuffer, float InValue, float* RESTRICT OutFloatBuffer, const int32 InNumSamples)
-	{
-#if !AUDIO_USE_SIMD
-		for (int32 i = 0; i < InNumSamples; ++i)
-		{
-			OutFloatBuffer[i] = InValue * InFloatBuffer[i];
-		}
-#else
-		RestrictedPtrAliasCheck(InFloatBuffer, OutFloatBuffer, InNumSamples);
-
-		// Can only SIMD on multiple of 4 buffers, we'll do normal multiples on last bit
-		const int32 NumSamplesRemaining = InNumSamples % AUDIO_NUM_FLOATS_PER_VECTOR_REGISTER;
-		const int32 NumSamplesToSimd = InNumSamples - NumSamplesRemaining;
-
-		// Load the single value we want to multiply all values by into a vector register
-		const VectorRegister4Float MultiplyValue = VectorLoadFloat1(&InValue);
-		for (int32 i = 0; i < NumSamplesToSimd; i += AUDIO_NUM_FLOATS_PER_VECTOR_REGISTER)
-		{
-			// Load the next 4 samples of the input buffer into a register
-			VectorRegister4Float InputBufferRegister = VectorLoadAligned(&InFloatBuffer[i]);
-
-			// Perform the multiply
-			VectorRegister4Float Temp = VectorMultiply(InputBufferRegister, MultiplyValue);
-
-			// Store results into the output buffer
-			VectorStoreAligned(Temp, &OutFloatBuffer[i]);
-		}
-
-		// Perform remaining non-simd values left over
-		for (int32 i = 0; i < NumSamplesRemaining; ++i)
-		{
-			OutFloatBuffer[NumSamplesToSimd + i] = InValue * InFloatBuffer[NumSamplesToSimd + i];
-		}
-#endif
-	}
-
-	void MultiplyBufferByConstantInPlace(FAlignedFloatBuffer& InBuffer, float InGain)
-	{
-		MultiplyBufferByConstantInPlace(InBuffer.GetData(), InBuffer.Num(), InGain);
-	}
-
-	void MultiplyBufferByConstantInPlace(float* RESTRICT InBuffer, int32 NumSamples, float InGain)
-	{
-		const VectorRegister4Float Gain = VectorLoadFloat1(&InGain);
-
-		for (int32 i = 0; i < NumSamples; i += AUDIO_NUM_FLOATS_PER_VECTOR_REGISTER)
-		{
-			VectorRegister4Float Output = VectorLoadAligned(&InBuffer[i]);
-			Output = VectorMultiply(Output, Gain);
-			VectorStoreAligned(Output, &InBuffer[i]);
-		}
-	}
-
-	// Adds a constant to a buffer (useful for DC offset removal)
-	void AddConstantToBufferInplace(FAlignedFloatBuffer& InBuffer, float InConstant)
-	{
-		AddConstantToBufferInplace(InBuffer.GetData(), InBuffer.Num(), InConstant);
-	}
-
-	void AddConstantToBufferInplace(float* RESTRICT InBuffer, int32 NumSamples, float InConstant)
-	{
-		const VectorRegister4Float Constant = VectorLoadFloat1(&InConstant);
-
-		for (int32 i = 0; i < NumSamples; i += AUDIO_NUM_FLOATS_PER_VECTOR_REGISTER)
-		{
-			VectorRegister4Float Output = VectorLoadAligned(&InBuffer[i]);
-			Output = VectorAdd(Output, Constant);
-			VectorStoreAligned(Output, &InBuffer[i]);
-		}
-	}
-
-	void BufferSetToConstantInplace(FAlignedFloatBuffer& InBuffer, float InConstant)
-	{
-		BufferSetToConstantInplace(InBuffer.GetData(), InBuffer.Num(), InConstant);
-	}
-
-	void BufferSetToConstantInplace(float* RESTRICT InBuffer, int32 NumSamples, float InConstant)
-	{
-		checkf(NumSamples >= AUDIO_NUM_FLOATS_PER_VECTOR_REGISTER, TEXT("Buffer must have at least %i elements."), AUDIO_NUM_FLOATS_PER_VECTOR_REGISTER);
-		checkf(0 == (NumSamples % AUDIO_NUM_FLOATS_PER_VECTOR_REGISTER), TEXT("Buffer length be a multiple of %i."), AUDIO_NUM_FLOATS_PER_VECTOR_REGISTER);
-		checkf(IsAligned<const float*>(InBuffer, AUDIO_SIMD_BYTE_ALIGNMENT), TEXT("Memory must be aligned to use vector operations."));
-
-		const VectorRegister4Float Constant = VectorLoadFloat1(&InConstant);
-
-		for (int32 i = 0; i < NumSamples; i += AUDIO_NUM_FLOATS_PER_VECTOR_REGISTER)
-		{
-			VectorStoreAligned(Constant, &InBuffer[i]);
-		}
-	}
-
-	/* Performs an element-wise weighted sum OutputBuffer = (InBuffer1 x InGain1) + (InBuffer2 x InGain2) */
-	void BufferWeightedSumFast(const FAlignedFloatBuffer& InBuffer1, float InGain1, const FAlignedFloatBuffer& InBuffer2, float InGain2, FAlignedFloatBuffer& OutBuffer)
-	{
-		checkf(InBuffer1.Num() == InBuffer2.Num(), TEXT("Buffers must be equal length"));
-		OutBuffer.Reset();
-		OutBuffer.AddUninitialized(InBuffer1.Num());
-
-		BufferWeightedSumFast(InBuffer1.GetData(), InGain1, InBuffer2.GetData(), InGain2, OutBuffer.GetData(), InBuffer1.Num());
-	}
-
-	/* Performs an element-wise weighted sum OutputBuffer = (InBuffer1 x InGain1) + InBuffer2 */
-	void BufferWeightedSumFast(const FAlignedFloatBuffer& InBuffer1, float InGain1, const FAlignedFloatBuffer& InBuffer2, FAlignedFloatBuffer& OutBuffer)
-	{
-		checkf(InBuffer1.Num() == InBuffer2.Num(), TEXT("Buffers must be equal length"));
-		OutBuffer.Reset();
-		OutBuffer.AddUninitialized(InBuffer1.Num());
-
-		BufferWeightedSumFast(InBuffer1.GetData(), InGain1, InBuffer2.GetData(), OutBuffer.GetData(), InBuffer1.Num());
-	}
-
-	/* Performs an element-wise weighted sum OutputBuffer = (InBuffer1 x InGain1) + (InBuffer2 x InGain2) */
-	void BufferWeightedSumFast(const float* RESTRICT InBuffer1, float InGain1, const float* RESTRICT InBuffer2, float InGain2, float* RESTRICT OutBuffer, int32 InNum)
-	{
-		checkf(InNum >= AUDIO_NUM_FLOATS_PER_VECTOR_REGISTER, TEXT("Buffer must have at least %i elements."), AUDIO_NUM_FLOATS_PER_VECTOR_REGISTER);
-		checkf(0 == (InNum % AUDIO_NUM_FLOATS_PER_VECTOR_REGISTER), TEXT("Buffer length be a multiple of %i."), AUDIO_NUM_FLOATS_PER_VECTOR_REGISTER);
-		checkf(IsAligned<const float*>(InBuffer1, AUDIO_SIMD_BYTE_ALIGNMENT), TEXT("Memory must be aligned to use vector operations."));
-		checkf(IsAligned<const float*>(InBuffer2, AUDIO_SIMD_BYTE_ALIGNMENT), TEXT("Memory must be aligned to use vector operations."));
-		checkf(IsAligned<float*>(OutBuffer, AUDIO_SIMD_BYTE_ALIGNMENT), TEXT("Memory must be aligned to use vector operations."));
-
-		RestrictedPtrAliasCheck(InBuffer1, OutBuffer, InNum);
-		RestrictedPtrAliasCheck(InBuffer2, OutBuffer, InNum);
-
-		VectorRegister4Float Gain1Vector = VectorLoadFloat1(&InGain1);
-		VectorRegister4Float Gain2Vector = VectorLoadFloat1(&InGain2);
-
-		for (int32 i = 0; i < InNum; i += AUDIO_NUM_FLOATS_PER_VECTOR_REGISTER)
-		{
-			// InBuffer1 x InGain1
-			VectorRegister4Float Input1 = VectorLoadAligned(&InBuffer1[i]);
-
-			// InBuffer2 x InGain2
-			VectorRegister4Float Input2 = VectorLoadAligned(&InBuffer2[i]);
-			VectorRegister4Float Weighted2 = VectorMultiply(Input2, Gain2Vector);
-
-			VectorRegister4Float Output = VectorMultiplyAdd(Input1, Gain1Vector, Weighted2);
-			VectorStoreAligned(Output, &OutBuffer[i]);
-		}
-	}
-
-	/* Performs an element-wise weighted sum OutputBuffer = (InBuffer1 x InGain1) + InBuffer2 */
-	void BufferWeightedSumFast(const float* RESTRICT InBuffer1, float InGain1, const float* RESTRICT InBuffer2, float* RESTRICT OutBuffer, int32 InNum)
-	{
-		checkf(InNum >= AUDIO_NUM_FLOATS_PER_VECTOR_REGISTER, TEXT("Buffer must have at least %i elements."), AUDIO_NUM_FLOATS_PER_VECTOR_REGISTER);
-		checkf(0 == (InNum % AUDIO_NUM_FLOATS_PER_VECTOR_REGISTER), TEXT("Buffer length be a multiple of %i."), AUDIO_NUM_FLOATS_PER_VECTOR_REGISTER);
-		checkf(IsAligned<const float*>(InBuffer1, AUDIO_SIMD_BYTE_ALIGNMENT), TEXT("Memory must be aligned to use vector operations."));
-		checkf(IsAligned<const float*>(InBuffer2, AUDIO_SIMD_BYTE_ALIGNMENT), TEXT("Memory must be aligned to use vector operations."));
-		checkf(IsAligned<float*>(OutBuffer, AUDIO_SIMD_BYTE_ALIGNMENT), TEXT("Memory must be aligned to use vector operations."));
-
-		RestrictedPtrAliasCheck(InBuffer1, OutBuffer, InNum);
-		RestrictedPtrAliasCheck(InBuffer2, OutBuffer, InNum);
-
-		VectorRegister4Float Gain1Vector = VectorLoadFloat1(&InGain1);
-
-		for (int32 i = 0; i < InNum; i += AUDIO_NUM_FLOATS_PER_VECTOR_REGISTER)
-		{
-			// InBuffer1 x InGain1
-			VectorRegister4Float Input1 = VectorLoadAligned(&InBuffer1[i]);
-			VectorRegister4Float Input2 = VectorLoadAligned(&InBuffer2[i]);
-			VectorRegister4Float Output = VectorMultiplyAdd(Input1, Gain1Vector, Input2);
-			VectorStoreAligned(Output, &OutBuffer[i]);
-		}
-	}
-
-	void FadeBufferFast(FAlignedFloatBuffer& OutFloatBuffer, const float StartValue, const float EndValue)
-	{
-		FadeBufferFast(OutFloatBuffer.GetData(), OutFloatBuffer.Num(), StartValue, EndValue);
-	}
-
-	void FadeBufferFast(float* RESTRICT OutFloatBuffer, int32 NumSamples, const float StartValue, const float EndValue)
-	{
-		checkf(IsAligned<float*>(OutFloatBuffer, AUDIO_SIMD_BYTE_ALIGNMENT), TEXT("Memory must be aligned to use vector operations."));
-		checkf(NumSamples % AUDIO_NUM_FLOATS_PER_VECTOR_REGISTER == 0, TEXT("Please use a buffer size that is a multiple of %i."), AUDIO_NUM_FLOATS_PER_VECTOR_REGISTER);
-
-#if !AUDIO_USE_SIMD
-		float Gain = StartValue;
-		if (FMath::IsNearlyEqual(StartValue, EndValue))
-		{
-			// No need to do anything if start and end values are both 0.0
-			if (StartValue == 0.0f)
-			{
-				FMemory::Memset(OutFloatBuffer, 0, sizeof(float)*NumSamples);
-			}
-			else
-			{
-				// Only need to do a buffer multiply if start and end values are the same
-				for (int32 i = 0; i < NumSamples; ++i)
-				{
-					OutFloatBuffer[i] = OutFloatBuffer[i] * Gain;
-				}
-			}
-		}
-		else
-		{
-			// Do a fade from start to end
-			const float DeltaValue = ((EndValue - StartValue) / NumSamples);
-			for (int32 i = 0; i < NumSamples; ++i)
-			{
-				OutFloatBuffer[i] = OutFloatBuffer[i] * Gain;
-				Gain += DeltaValue;
-			}
-		}
-#else
-		const int32 NumIterations = NumSamples / 4;
-
-		if (FMath::IsNearlyEqual(StartValue, EndValue))
-		{
-			// No need to do anything if start and end values are both 0.0
-			if (StartValue == 0.0f)
-			{
-				FMemory::Memset(OutFloatBuffer, 0, sizeof(float)*NumSamples);
-			}
-			else
-			{
-				VectorRegister4Float Gain = VectorLoadFloat1(&StartValue);
-
-				for (int32 i = 0; i < NumSamples; i += AUDIO_NUM_FLOATS_PER_VECTOR_REGISTER)
-				{
-					VectorRegister4Float Output = VectorLoadAligned(&OutFloatBuffer[i]);
-					Output = VectorMultiply(Output, Gain);
-					VectorStoreAligned(Output, &OutFloatBuffer[i]);
-				}
-			}
-		}
-		else
-		{
-			const float DeltaValue = ((EndValue - StartValue) / NumIterations);
-
-			VectorRegister4Float Gain = VectorLoadFloat1(&StartValue);
-			VectorRegister4Float Delta = VectorLoadFloat1(&DeltaValue);
-
-			for (int32 i = 0; i < NumSamples; i += AUDIO_NUM_FLOATS_PER_VECTOR_REGISTER)
-			{
-				VectorRegister4Float Output = VectorLoadAligned(&OutFloatBuffer[i]);
-				Output = VectorMultiply(Output, Gain);
-				Gain = VectorAdd(Gain, Delta);
-				VectorStoreAligned(Output, &OutFloatBuffer[i]);
-			}
-		}
-#endif
-	}
-
-	void MixInBufferFast(const FAlignedFloatBuffer& InFloatBuffer, FAlignedFloatBuffer& BufferToSumTo, const float Gain)
-	{
-		MixInBufferFast(InFloatBuffer.GetData(), BufferToSumTo.GetData(), InFloatBuffer.Num(), Gain);
-	}
-
-	void MixInBufferFast(const float* RESTRICT InFloatBuffer, float* RESTRICT BufferToSumTo, int32 NumSamples, const float Gain)
-	{
-		checkf(IsAligned<const float*>(InFloatBuffer, AUDIO_SIMD_BYTE_ALIGNMENT), TEXT("Memory must be aligned to use vector operations."));
-		checkf(IsAligned<float*>(BufferToSumTo, AUDIO_SIMD_BYTE_ALIGNMENT), TEXT("Memory must be aligned to use vector operations."));
-		checkf(NumSamples % AUDIO_NUM_FLOATS_PER_VECTOR_REGISTER == 0, TEXT("Please use a buffer size that is a multiple of %i."), AUDIO_NUM_FLOATS_PER_VECTOR_REGISTER);
-
-		RestrictedPtrAliasCheck(InFloatBuffer, BufferToSumTo, NumSamples);
-
-#if !AUDIO_USE_SIMD
-		for (int32 i = 0; i < NumSamples; ++i)
-		{
-			BufferToSumTo[i] += InFloatBuffer[i] * Gain;
-		}
-#else
-		VectorRegister4Float GainVector = VectorLoadFloat1(&Gain);
-
-		for (int32 i = 0; i < NumSamples; i += AUDIO_NUM_FLOATS_PER_VECTOR_REGISTER)
-		{
-			VectorRegister4Float Output = VectorLoadAligned(&BufferToSumTo[i]);
-			VectorRegister4Float Input  = VectorLoadAligned(&InFloatBuffer[i]);
-			Output = VectorMultiplyAdd(Input, GainVector, Output);
-			VectorStoreAligned(Output, &BufferToSumTo[i]);
-		}
-#endif
-	}
-
-	void MixInBufferFast(const FAlignedFloatBuffer& InFloatBuffer, FAlignedFloatBuffer& BufferToSumTo)
-	{
-		checkf(InFloatBuffer.Num() == BufferToSumTo.Num(), TEXT("Buffers must be equal size"));
-		MixInBufferFast(InFloatBuffer.GetData(), BufferToSumTo.GetData(), InFloatBuffer.Num());
-	}
-
-	void MixInBufferFast(const float* RESTRICT InFloatBuffer, float* RESTRICT BufferToSumTo, int32 NumSamples)
-	{
-		checkf(IsAligned<const float*>(InFloatBuffer, AUDIO_SIMD_BYTE_ALIGNMENT), TEXT("Memory must be aligned to use vector operations."));
-		checkf(IsAligned<float*>(BufferToSumTo, AUDIO_SIMD_BYTE_ALIGNMENT), TEXT("Memory must be aligned to use vector operations."));
-		checkf(NumSamples % AUDIO_NUM_FLOATS_PER_VECTOR_REGISTER == 0, TEXT("Please use a buffer size that is a multiple of %i."), AUDIO_NUM_FLOATS_PER_VECTOR_REGISTER);
-
-		RestrictedPtrAliasCheck(InFloatBuffer, BufferToSumTo, NumSamples);
-
-#if !AUDIO_USE_SIMD
-		for (int32 i = 0; i < NumSamples; ++i)
-		{
-			BufferToSumTo[i] += InFloatBuffer[i];
-		}
-#else
-		for (int32 i = 0; i < NumSamples; i += AUDIO_NUM_FLOATS_PER_VECTOR_REGISTER)
-		{
-			VectorRegister4Float Output = VectorLoadAligned(&BufferToSumTo[i]);
-			VectorRegister4Float Input = VectorLoadAligned(&InFloatBuffer[i]);
-			Output = VectorAdd(Input, Output);
-			VectorStoreAligned(Output, &BufferToSumTo[i]);
-		}
-#endif
-	}
-
-	void MixInBufferFast(const FAlignedFloatBuffer& InFloatBuffer, FAlignedFloatBuffer& BufferToSumTo, const float StartGain, const float EndGain)
-	{
-		MixInBufferFast(InFloatBuffer.GetData(), BufferToSumTo.GetData(), InFloatBuffer.Num(), StartGain, EndGain);
-	}
-
-	void MixInBufferFast(const float* RESTRICT InFloatBuffer, float* RESTRICT BufferToSumTo, int32 NumSamples, const float StartGain, const float EndGain)
-	{
-		checkf(IsAligned<const float*>(InFloatBuffer, AUDIO_SIMD_BYTE_ALIGNMENT), TEXT("Memory must be aligned to use vector operations."));
-		checkf(IsAligned<float*>(BufferToSumTo, AUDIO_SIMD_BYTE_ALIGNMENT), TEXT("Memory must be aligned to use vector operations."));
-		checkf(NumSamples % AUDIO_NUM_FLOATS_PER_VECTOR_REGISTER == 0, TEXT("Please use a buffer size that is a multiple of %i."), AUDIO_NUM_FLOATS_PER_VECTOR_REGISTER);
-
-		RestrictedPtrAliasCheck(InFloatBuffer, BufferToSumTo, NumSamples);
-
-		const int32 NumIterations = NumSamples / 4;
-
-		if (FMath::IsNearlyEqual(StartGain, EndGain))
-		{
-			// No need to do anything if start and end values are both 0.0
-			if (StartGain == 0.0f)
-			{
-				return;
-			}
-			else
-			{
-				VectorRegister4Float Gain = VectorLoadFloat1(&StartGain);
-
-				for (int32 i = 0; i < NumSamples; i += AUDIO_NUM_FLOATS_PER_VECTOR_REGISTER)
-				{
-					VectorRegister4Float Input = VectorLoadAligned(&InFloatBuffer[i]);
-					VectorRegister4Float Output = VectorLoadAligned(&BufferToSumTo[i]);
-
-					Output = VectorMultiplyAdd(Input, Gain, Output);					
-					VectorStoreAligned(Output, &BufferToSumTo[i]);
-				}
-			}
-		}
-		else
-		{
-			const float DeltaValue = ((EndGain - StartGain) / NumIterations);
-
-			VectorRegister4Float Gain = VectorLoadFloat1(&StartGain);
-			VectorRegister4Float Delta = VectorLoadFloat1(&DeltaValue);
-
-			for (int32 i = 0; i < NumSamples; i += AUDIO_NUM_FLOATS_PER_VECTOR_REGISTER)
-			{
-				VectorRegister4Float Input = VectorLoadAligned(&InFloatBuffer[i]);
-				VectorRegister4Float Output = VectorLoadAligned(&BufferToSumTo[i]);
-				Output = VectorMultiplyAdd(Input, Gain, Output);
-
-				VectorStoreAligned(Output, &BufferToSumTo[i]);
-
-				Gain = VectorAdd(Gain, Delta);
-			}
-		}
-	}
-
-	/* Subtracts two buffers together element-wise. */
-	void BufferSubtractFast(const FAlignedFloatBuffer& InMinuend, const FAlignedFloatBuffer& InSubtrahend, FAlignedFloatBuffer& OutputBuffer)
-	{
-		const int32 InNum = InMinuend.Num();
-		OutputBuffer.Reset(InNum);
-		OutputBuffer.AddUninitialized(InNum);
-
-		checkf(InMinuend.Num() == InSubtrahend.Num(), TEXT("Input buffers must be equal length. InMinuend.Num()=%d, InSubtrahend.Num()=%d"), InMinuend.Num(), InSubtrahend.Num());
-
-		BufferSubtractFast(InMinuend.GetData(), InSubtrahend.GetData(), OutputBuffer.GetData(), OutputBuffer.Num());
-	}
-
-	/* Subtracts two buffers together element-wise. */
-	void BufferSubtractFast(const float* RESTRICT InMinuend, const float* RESTRICT InSubtrahend, float* RESTRICT OutBuffer, int32 InNum)
-	{
-		checkf(InNum >= 4, TEXT("Buffer must have at least %i elements."), AUDIO_NUM_FLOATS_PER_VECTOR_REGISTER);
-		checkf(0 == (InNum % AUDIO_NUM_FLOATS_PER_VECTOR_REGISTER), TEXT("Buffer length be a multiple of %i."), AUDIO_NUM_FLOATS_PER_VECTOR_REGISTER);
-		checkf(IsAligned<const float*>(InMinuend, AUDIO_SIMD_BYTE_ALIGNMENT), TEXT("Memory must be aligned to use vector operations."));
-		checkf(IsAligned<const float*>(InSubtrahend, AUDIO_SIMD_BYTE_ALIGNMENT), TEXT("Memory must be aligned to use vector operations."));
-		checkf(IsAligned<float*>(OutBuffer, AUDIO_SIMD_BYTE_ALIGNMENT), TEXT("Memory must be aligned to use vector operations."));
-
-		RestrictedPtrAliasCheck(InMinuend, OutBuffer, InNum);
-		RestrictedPtrAliasCheck(InSubtrahend, OutBuffer, InNum);
-
-		for (int32 i = 0; i < InNum; i += AUDIO_NUM_FLOATS_PER_VECTOR_REGISTER)
-		{
-			VectorRegister4Float Input1 = VectorLoadAligned(&InMinuend[i]);
-			VectorRegister4Float Input2 = VectorLoadAligned(&InSubtrahend[i]);
-			VectorRegister4Float Output = VectorSubtract(Input1, Input2);
-			VectorStoreAligned(Output, &OutBuffer[i]);
-		}
-	}
-
-	/* Performs element-wise in-place subtraction placing the result in the subtrahend. InOutSubtrahend = InMinuend - InOutSubtrahend */
-	void BufferSubtractInPlace1Fast(const FAlignedFloatBuffer& InMinuend, FAlignedFloatBuffer& InOutSubtrahend)
-	{
-		checkf(InMinuend.Num() == InOutSubtrahend.Num(), TEXT("Input buffers must be equal length"));
-		BufferSubtractInPlace1Fast(InMinuend.GetData(), InOutSubtrahend.GetData(), InMinuend.Num());
-	}
-
-	/* Performs element-wise in-place subtraction placing the result in the subtrahend. InOutSubtrahend = InMinuend - InOutSubtrahend */
-	void BufferSubtractInPlace1Fast(const float* RESTRICT InMinuend, float* RESTRICT InOutSubtrahend, int32 InNum)
-	{
-		checkf(InNum >= 4, TEXT("Buffer must have at least %i elements."), AUDIO_NUM_FLOATS_PER_VECTOR_REGISTER);
-		checkf(0 == (InNum % AUDIO_NUM_FLOATS_PER_VECTOR_REGISTER), TEXT("Buffer length be a multiple of %i."), AUDIO_NUM_FLOATS_PER_VECTOR_REGISTER);
-		checkf(IsAligned<const float*>(InMinuend, AUDIO_SIMD_BYTE_ALIGNMENT), TEXT("Memory must be aligned to use vector operations."));
-		checkf(IsAligned<const float*>(InOutSubtrahend, AUDIO_SIMD_BYTE_ALIGNMENT), TEXT("Memory must be aligned to use vector operations."));
-		
-		RestrictedPtrAliasCheck(InMinuend, InOutSubtrahend, InNum);
-
-		for (int32 i = 0; i < InNum; i += AUDIO_NUM_FLOATS_PER_VECTOR_REGISTER)
-		{
-			VectorRegister4Float Input1 = VectorLoadAligned(&InMinuend[i]);
-			VectorRegister4Float Input2 = VectorLoadAligned(&InOutSubtrahend[i]);
-
-			VectorRegister4Float Output = VectorSubtract(Input1, Input2);
-			VectorStoreAligned(Output, &InOutSubtrahend[i]);
-		}
-	}
-	
-	/* Performs element-wise in-place subtraction placing the result in the minuend. InOutMinuend = InOutMinuend - InSubtrahend */
-	void BufferSubtractInPlace2Fast(FAlignedFloatBuffer& InOutMinuend, const FAlignedFloatBuffer& InSubtrahend)
-	{
-		checkf(InOutMinuend.Num() == InSubtrahend.Num(), TEXT("Input buffers must be equal length"));
-		BufferSubtractInPlace2Fast(InOutMinuend.GetData(), InSubtrahend.GetData(), InOutMinuend.Num());
-	}
-
-	/* Performs element-wise in-place subtraction placing the result in the minuend. InOutMinuend = InOutMinuend - InSubtrahend */
-	void BufferSubtractInPlace2Fast(float* RESTRICT InOutMinuend, const float* RESTRICT InSubtrahend, int32 InNum)
-	{
-		checkf(InNum >= 4, TEXT("Buffer must have at least %i elements."), AUDIO_NUM_FLOATS_PER_VECTOR_REGISTER);
-		checkf(0 == (InNum % AUDIO_NUM_FLOATS_PER_VECTOR_REGISTER), TEXT("Buffer length be a multiple of %i."), AUDIO_NUM_FLOATS_PER_VECTOR_REGISTER);
-		checkf(IsAligned<const float*>(InOutMinuend, AUDIO_SIMD_BYTE_ALIGNMENT), TEXT("Memory must be aligned to use vector operations."));
-		checkf(IsAligned<const float*>(InSubtrahend, AUDIO_SIMD_BYTE_ALIGNMENT), TEXT("Memory must be aligned to use vector operations."));
-
-		RestrictedPtrAliasCheck(InOutMinuend, InSubtrahend, InNum);
-		
-		for (int32 i = 0; i < InNum; i += AUDIO_NUM_FLOATS_PER_VECTOR_REGISTER)
-		{
-			VectorRegister4Float Input1 = VectorLoadAligned(&InOutMinuend[i]);
-			VectorRegister4Float Input2 = VectorLoadAligned(&InSubtrahend[i]);
-
-			VectorRegister4Float Output = VectorSubtract(Input1, Input2);
-			VectorStoreAligned(Output, &InOutMinuend[i]);
-		}
-	}
-
-	void SumBuffers(const FAlignedFloatBuffer& InFloatBuffer1, const FAlignedFloatBuffer& InFloatBuffer2, FAlignedFloatBuffer& OutputBuffer)
-	{
-		checkf(InFloatBuffer1.Num() == InFloatBuffer2.Num(), TEXT("Input buffers must be equal length"));
-		const int32 InNum = InFloatBuffer1.Num();
-		OutputBuffer.Reset(InNum);
-		OutputBuffer.AddUninitialized(InNum);
-
-		SumBuffers(InFloatBuffer1.GetData(), InFloatBuffer2.GetData(), OutputBuffer.GetData(), InNum);
-	}
-
-	void SumBuffers(const float* RESTRICT InFloatBuffer1, const float* RESTRICT InFloatBuffer2, float* RESTRICT OutputBuffer, int32 NumSamples)
-	{
-		checkf(IsAligned<const float*>(InFloatBuffer1, AUDIO_SIMD_BYTE_ALIGNMENT), TEXT("Memory must be aligned to use vector operations."));
-		checkf(IsAligned<const float*>(InFloatBuffer2, AUDIO_SIMD_BYTE_ALIGNMENT), TEXT("Memory must be aligned to use vector operations."));
-		checkf(IsAligned<float*>(OutputBuffer, AUDIO_SIMD_BYTE_ALIGNMENT), TEXT("Memory must be aligned to use vector operations."));
-		checkf(NumSamples % AUDIO_NUM_FLOATS_PER_VECTOR_REGISTER == 0, TEXT("Please use a buffer size that is a multiple of %i."), AUDIO_NUM_FLOATS_PER_VECTOR_REGISTER);
-
-		RestrictedPtrAliasCheck(InFloatBuffer1, OutputBuffer, NumSamples);
-		RestrictedPtrAliasCheck(InFloatBuffer2, OutputBuffer, NumSamples);
-
-#if !AUDIO_USE_SIMD
-		for (int32 i = 0; i < NumSamples; ++i)
-		{
-			OutputBuffer[i] = InFloatBuffer1[i] + InFloatBuffer2[i];
-		}
-#else
-		for (int32 i = 0; i < NumSamples; i += AUDIO_NUM_FLOATS_PER_VECTOR_REGISTER)
-		{
-			VectorRegister4Float Input1 = VectorLoadAligned(&InFloatBuffer1[i]);
-			VectorRegister4Float Input2 = VectorLoadAligned(&InFloatBuffer2[i]);
-
-			VectorRegister4Float Output = VectorAdd(Input1, Input2);
-			VectorStoreAligned(Output, &OutputBuffer[i]);
-		}
-#endif
-	}
-
-	void MultiplyBuffersInPlace(const FAlignedFloatBuffer& InFloatBuffer, FAlignedFloatBuffer& BufferToMultiply)
-	{
-		MultiplyBuffersInPlace(InFloatBuffer.GetData(), BufferToMultiply.GetData(), BufferToMultiply.Num());
-	}
-
-	void MultiplyBuffersInPlace(const float* RESTRICT InFloatBuffer, float* RESTRICT BufferToMultiply, int32 NumSamples)
-	{
-		checkf(IsAligned<const float*>(InFloatBuffer, AUDIO_SIMD_BYTE_ALIGNMENT), TEXT("Memory must be aligned to use vector operations."));
-		checkf(IsAligned<const float*>(BufferToMultiply, AUDIO_SIMD_BYTE_ALIGNMENT), TEXT("Memory must be aligned to use vector operations."));
-		checkf(NumSamples % AUDIO_NUM_FLOATS_PER_VECTOR_REGISTER == 0, TEXT("Please use a buffer size that is a multiple of %i."), AUDIO_NUM_FLOATS_PER_VECTOR_REGISTER);
-
-		RestrictedPtrAliasCheck(InFloatBuffer, BufferToMultiply, NumSamples);
-
-		for (int32 i = 0; i < NumSamples; i += AUDIO_NUM_FLOATS_PER_VECTOR_REGISTER)
-		{
-			VectorRegister4Float Input1 = VectorLoadAligned(&InFloatBuffer[i]);
-			VectorRegister4Float Output = VectorLoadAligned(&BufferToMultiply[i]);
-
-			Output = VectorMultiply(Input1, Output);
-			VectorStoreAligned(Output, &BufferToMultiply[i]);
-		}
-	}
-
-	float GetMagnitude(const FAlignedFloatBuffer& Buffer)
-	{
-		return GetMagnitude(Buffer.GetData(), Buffer.Num());
-	}
-
-	float GetMagnitude(const float* RESTRICT Buffer, int32 NumSamples)
-	{
-		checkf(NumSamples % AUDIO_NUM_FLOATS_PER_VECTOR_REGISTER == 0, TEXT("Please use a buffer size that is a multiple of %i."), AUDIO_NUM_FLOATS_PER_VECTOR_REGISTER);
-
-#if !AUDIO_USE_SIMD
-		float Sum = 0.0f;
-		for (int32 i = 0; i < NumSamples; ++i)
-		{
-			Sum += Buffer[i] * Buffer[i];
-		}
-		return FMath::Sqrt(Sum);
-#else
-		VectorRegister4Float Sum = VectorZero();
-
-		const float Exponent = 2.0f;
-		VectorRegister4Float ExponentVector = VectorLoadFloat1(&Exponent);
-
-		for (int32 i = 0; i < NumSamples; i += AUDIO_NUM_FLOATS_PER_VECTOR_REGISTER)
-		{
-			VectorRegister4Float Input = VectorPow(VectorLoadAligned(&Buffer[i]), ExponentVector);
-			Sum = VectorAdd(Sum, Input);
-		}
-
-		float PartionedSums[4];
-		VectorStoreAligned(Sum, PartionedSums);
-
-		return FMath::Sqrt(PartionedSums[0] + PartionedSums[1] + PartionedSums[2] + PartionedSums[3]);
-#endif
-	}
-
-	float BufferGetAverageValue(const FAlignedFloatBuffer& Buffer)
-	{
-		checkf(Buffer.Num() % AUDIO_NUM_FLOATS_PER_VECTOR_REGISTER == 0, TEXT("Please use a buffer size that is a multiple of %i."), AUDIO_NUM_FLOATS_PER_VECTOR_REGISTER);
-
-		return BufferGetAverageValue(Buffer.GetData(), Buffer.Num());
-	}
-
-	float BufferGetAverageValue(const float* RESTRICT Buffer, int32 NumSamples)
-	{
-		checkf(NumSamples % AUDIO_NUM_FLOATS_PER_VECTOR_REGISTER == 0, TEXT("Please use a buffer size that is a multiple of %i."), AUDIO_NUM_FLOATS_PER_VECTOR_REGISTER);
-
-#if !AUDIO_USE_SIMD
-		float Sum = 0.0f;
-		for (int32 i = 0; i < NumSamples; ++i)
-		{
-			Sum += Buffer[i];
-		}
-		return Sum / NumSamples;
-#else
-		VectorRegister4Float Sum = VectorZero();
-
-		for (int32 i = 0; i < NumSamples; i += AUDIO_NUM_FLOATS_PER_VECTOR_REGISTER)
-		{
-			VectorRegister4Float Input = VectorLoadAligned(&Buffer[i]);
-			Sum = VectorAdd(Sum, Input);
-		}
-
-		float PartionedSums[4];
-		VectorStore(Sum, PartionedSums);
-		
-		return (PartionedSums[0] + PartionedSums[1] + PartionedSums[2] + PartionedSums[3]) / NumSamples;
-#endif
-	}
-
-	float BufferGetAverageAbsValue(const FAlignedFloatBuffer& Buffer)
-	{
-		checkf(Buffer.Num() % AUDIO_NUM_FLOATS_PER_VECTOR_REGISTER == 0, TEXT("Please use a buffer size that is a multiple of %i."), AUDIO_NUM_FLOATS_PER_VECTOR_REGISTER);
-
-		return BufferGetAverageAbsValue(Buffer.GetData(), Buffer.Num());
-	}
-
-	float BufferGetAverageAbsValue(const float* RESTRICT Buffer, int32 NumSamples)
-	{
-		checkf(NumSamples % AUDIO_NUM_FLOATS_PER_VECTOR_REGISTER == 0, TEXT("Please use a buffer size that is a multiple of %i."), AUDIO_NUM_FLOATS_PER_VECTOR_REGISTER);
-
-#if !AUDIO_USE_SIMD
-		float Sum = 0.0f;
-		for (int32 i = 0; i < NumSamples; ++i)
-		{
-			Sum += FMath::Abs(Buffer[i]);
-		}
-		return Sum / NumSamples;
-#else
-		VectorRegister4Float Sum = VectorZero();
-
-		for (int32 i = 0; i < NumSamples; i += AUDIO_NUM_FLOATS_PER_VECTOR_REGISTER)
-		{
-			VectorRegister4Float Input = VectorAbs(VectorLoadAligned(&Buffer[i]));
-			Sum = VectorAdd(Sum, Input);
-		}
-
-		float PartionedSums[4];
-		VectorStore(Sum, PartionedSums);
-
-		return (PartionedSums[0] + PartionedSums[1] + PartionedSums[2] + PartionedSums[3]) / NumSamples;
-#endif
-	}
-
 	/**
 	 * CHANNEL MIXING OPERATIONS:
 	 * To understand these functions, it's best that you have prior experience reading SIMD code.
@@ -2000,42 +1307,6 @@ namespace Audio
 		}
 	}
 
-	void BufferComplexToPowerFast(const FAlignedFloatBuffer& InRealSamples, const FAlignedFloatBuffer& InImaginarySamples, FAlignedFloatBuffer& OutPowerSamples)
-	{
-		checkf(InRealSamples.Num() == InImaginarySamples.Num(), TEXT("Input buffers must have equal number of elements"));
-
-		const int32 Num = InRealSamples.Num();
-
-		OutPowerSamples.Reset(Num);
-		OutPowerSamples.AddUninitialized(Num);
-
-		BufferComplexToPowerFast(InRealSamples.GetData(), InImaginarySamples.GetData(), OutPowerSamples.GetData(), Num);
-	}
-
-	void BufferComplexToPowerFast(const float* RESTRICT InRealSamples, const float* RESTRICT InImaginarySamples, float* RESTRICT OutPowerSamples, const int32 InNum)
-	{
-		checkf(InNum >= 4, TEXT("Buffer must have at least %i elements."), AUDIO_NUM_FLOATS_PER_VECTOR_REGISTER);
-		checkf(0 == (InNum % AUDIO_NUM_FLOATS_PER_VECTOR_REGISTER), TEXT("Buffer length be a multiple of %i."), AUDIO_NUM_FLOATS_PER_VECTOR_REGISTER);
-
-		checkf(IsAligned<const float*>(InRealSamples, AUDIO_SIMD_BYTE_ALIGNMENT), TEXT("Memory must be aligned to use vector operations."));
-		checkf(IsAligned<const float*>(InImaginarySamples, AUDIO_SIMD_BYTE_ALIGNMENT), TEXT("Memory must be aligned to use vector operations."));
-		checkf(IsAligned<float*>(OutPowerSamples, AUDIO_SIMD_BYTE_ALIGNMENT), TEXT("Memory must be aligned to use vector operations."));
-
-		RestrictedPtrAliasCheck(InRealSamples, OutPowerSamples, InNum);
-		RestrictedPtrAliasCheck(InImaginarySamples, OutPowerSamples, InNum);
-
-		for (int32 i = 0; i < InNum; i += AUDIO_NUM_FLOATS_PER_VECTOR_REGISTER)
-		{
-			VectorRegister4Float VInReal = VectorLoadAligned(&InRealSamples[i]);
-			VectorRegister4Float VInRealSquared = VectorMultiply(VInReal, VInReal);
-
-			VectorRegister4Float VInImag = VectorLoadAligned(&InImaginarySamples[i]);
-			VectorRegister4Float VOut = VectorMultiplyAdd(VInImag, VInImag, VInRealSquared);
-
-			VectorStoreAligned(VOut, &OutPowerSamples[i]);
-		}
-	}
-
 	// class FBufferLinearEase implementation
 	FBufferLinearEase::FBufferLinearEase() {}
 	FBufferLinearEase::FBufferLinearEase(const FAlignedFloatBuffer & InSourceValues, const FAlignedFloatBuffer & InTargetValues, int32 InLerpLength)
@@ -2063,8 +1334,8 @@ namespace Audio
 
 		const float OneOverLerpLength = 1.0f / static_cast<float>(LerpLength);
 
-		BufferSubtractFast(InTargetValues.GetData(), InSourceValues.GetData(), DeltaBuffer.GetData(), BufferLength);
-		MultiplyBufferByConstantInPlace(DeltaBuffer, OneOverLerpLength);
+		ArraySubtract(InTargetValues, InSourceValues, DeltaBuffer);
+		ArrayMultiplyByConstantInPlace(DeltaBuffer, OneOverLerpLength);
 	}
 
 	bool FBufferLinearEase::Update(FAlignedFloatBuffer & InSourceValues)
@@ -2072,7 +1343,7 @@ namespace Audio
 		check(InSourceValues.Num() == BufferLength);
 		check(CurrentLerpStep != LerpLength);
 
-		MixInBufferFast(DeltaBuffer.GetData(), InSourceValues.GetData(), BufferLength);
+		ArrayMixIn(DeltaBuffer, InSourceValues);
 
 		if (++CurrentLerpStep == LerpLength)
 		{
@@ -2096,7 +1367,7 @@ namespace Audio
 			bIsComplete = true;
 		}
 
-		MixInBufferFast(DeltaBuffer.GetData(), InSourceValues.GetData(), BufferLength, static_cast<float>(StepsToJumpForward));
+		ArrayMixIn(DeltaBuffer, InSourceValues, static_cast<float>(StepsToJumpForward));
 
 		return bIsComplete;
 
diff --git a/Engine/Source/Runtime/SignalProcessing/Private/ConvertDeinterleave.cpp b/Engine/Source/Runtime/SignalProcessing/Private/ConvertDeinterleave.cpp
--- a/Engine/Source/Runtime/SignalProcessing/Private/ConvertDeinterleave.cpp
+++ b/Engine/Source/Runtime/SignalProcessing/Private/ConvertDeinterleave.cpp
@@ -124,16 +124,16 @@ namespace Audio
 				checkf(NumOutputChannels == OutSamples.Num(), TEXT("Output audio buffer not initialized to expected channel count."));
 
 				const int32 NumFrames = GetMultichannelBufferNumFrames(OutSamples);
-				const float* InSampleData = InSamples.GetData();
 
-				const int32 DataSize = NumFrames * sizeof(float);
 				for (int32 OutChannelIndex = 0; OutChannelIndex < NumOutputChannels; OutChannelIndex++)
 				{
 					const float Gain = ChannelGains[OutChannelIndex];
 					if (Gain != 0.f)
 					{
+						TArrayView<float> OutSampleView(OutSamples[OutChannelIndex].GetData(), NumFrames);
+
 						// Only need to multiply if gain is non-zero
-						Audio::BufferMultiplyByConstant(InSampleData, ChannelGains[OutChannelIndex], OutSamples[OutChannelIndex].GetData(), NumFrames);
+						Audio::ArrayMultiplyByConstant(InSamples, ChannelGains[OutChannelIndex], OutSampleView);
 					}
 					else
 					{
diff --git a/Engine/Source/Runtime/SignalProcessing/Private/Dsp.cpp b/Engine/Source/Runtime/SignalProcessing/Private/Dsp.cpp
--- a/Engine/Source/Runtime/SignalProcessing/Private/Dsp.cpp
+++ b/Engine/Source/Runtime/SignalProcessing/Private/Dsp.cpp
@@ -28,7 +28,7 @@ namespace Audio
 		ArraySubtract(InLeftChannel, InRightChannel, OutSideChannel);
 
 		//Output
-		SumBuffers(InLeftChannel, InRightChannel, OutMidChannel);
+		ArraySum(InLeftChannel, InRightChannel, OutMidChannel);
 
 	}
 
@@ -55,10 +55,10 @@ namespace Audio
 
 		// Calculate the average value between the two signals at each sample
 		const float AverageScale = 0.5f;
-		BufferWeightedSumFast(InMidChannel, AverageScale, InSideChannel, AverageScale, OutLeftChannel);
+		ArrayWeightedSum(InMidChannel, AverageScale, InSideChannel, AverageScale, OutLeftChannel);
 
 		ArraySubtract(InMidChannel, InSideChannel, OutRightChannel);
-		MultiplyBufferByConstantInPlace(OutRightChannel, AverageScale);
+		ArrayMultiplyByConstantInPlace(OutRightChannel, AverageScale);
 
 	}
 
diff --git a/Engine/Source/Runtime/SignalProcessing/Private/DynamicDelayAPF.cpp b/Engine/Source/Runtime/SignalProcessing/Private/DynamicDelayAPF.cpp
--- a/Engine/Source/Runtime/SignalProcessing/Private/DynamicDelayAPF.cpp
+++ b/Engine/Source/Runtime/SignalProcessing/Private/DynamicDelayAPF.cpp
@@ -105,27 +105,33 @@ void FDynamicDelayAPF::ProcessAudioBlock(const float* InSamples, const FAlignedF
 	DelayLineInput.Reset(InNum);
 	DelayLineInput.AddUninitialized(InNum);
 
+	TArrayView<float> WorkBufferAView(WorkBufferA.GetData(), InNum);
+	TArrayView<float> WorkBufferBView(WorkBufferB.GetData(), InNum);
+	TArrayView<float> DelayLineInputView(DelayLineInput.GetData(), InNum);
+	TArrayView<const float> InSamplesView(InSamples, InNum);
+	TArrayView<float> OutSamplesView(OutSamples, InNum);
+
 	// Get G values to interpolate over
 	const float LastG = G.GetNextValue();
 	const float CurrG = G.GetNextValue(InNum - 1);
 
 	// WorkBufferA = G * WorkBufferB
 	FMemory::Memcpy(WorkBufferA.GetData(), WorkBufferB.GetData(), InNum * sizeof(float));
-	FadeBufferFast(WorkBufferA.GetData(), InNum, LastG, CurrG);
+	ArrayFade(WorkBufferAView, LastG, CurrG);
 
 	// DelayLineInput = InSamples + WorkBufferA
 	// w[n] = x[n] + G * w[n - d]
-	SumBuffers(InSamples, WorkBufferA.GetData(), DelayLineInput.GetData(), InNum);
+	ArraySum(InSamplesView, WorkBufferAView, DelayLineInputView);
 
-	BufferUnderflowClampFast(DelayLineInput);
+	ArrayUnderflowClamp(DelayLineInput);
 
 	// Update delay line
 	IntegerDelayLine->RemoveSamples(InNum);
 	IntegerDelayLine->AddSamples(DelayLineInput.GetData(), InNum);
 
 	// y[n] = -G w[n] + w[n - d]
-	FadeBufferFast(DelayLineInput.GetData(), InNum, -LastG, -CurrG);
-	SumBuffers(DelayLineInput.GetData(), WorkBufferB.GetData(), OutSamples, InNum);
+	ArrayFade(DelayLineInputView, -LastG, -CurrG);
+	ArraySum(DelayLineInputView, WorkBufferBView, OutSamplesView);
 }
 
 void FDynamicDelayAPF::Reset() 
diff --git a/Engine/Source/Runtime/SignalProcessing/Private/EarlyReflectionsFast.cpp b/Engine/Source/Runtime/SignalProcessing/Private/EarlyReflectionsFast.cpp
--- a/Engine/Source/Runtime/SignalProcessing/Private/EarlyReflectionsFast.cpp
+++ b/Engine/Source/Runtime/SignalProcessing/Private/EarlyReflectionsFast.cpp
@@ -216,8 +216,8 @@ void FEarlyReflectionsFast::ProcessAudio(const FAlignedFloatBuffer& InSamples, c
 	const float OneMinusGain = 1.0 - Settings.Gain;
 
 	// Apply Gain
-	MultiplyBufferByConstantInPlace(OutLeftSamples, Settings.Gain);
-	MultiplyBufferByConstantInPlace(OutRightSamples, Settings.Gain);
+	ArrayMultiplyByConstantInPlace(OutLeftSamples, Settings.Gain);
+	ArrayMultiplyByConstantInPlace(OutRightSamples, Settings.Gain);
 }
 
 void FEarlyReflectionsFast::FlushAudio()
diff --git a/Engine/Source/Runtime/SignalProcessing/Private/Flanger.cpp b/Engine/Source/Runtime/SignalProcessing/Private/Flanger.cpp
--- a/Engine/Source/Runtime/SignalProcessing/Private/Flanger.cpp
+++ b/Engine/Source/Runtime/SignalProcessing/Private/Flanger.cpp
@@ -81,6 +81,6 @@ namespace Audio
 			ScratchBuffer[SampleIndex] = DelayBuffer.ProcessAudioSample(InBuffer[SampleIndex]);
 		}
 
-		BufferWeightedSumFast(ScratchBuffer, DelayedSignalLevel, InBuffer, NonDelayedSignalLevel, OutBuffer);
+		ArrayWeightedSum(ScratchBuffer, DelayedSignalLevel, InBuffer, NonDelayedSignalLevel, OutBuffer);
 	}
 }
diff --git a/Engine/Source/Runtime/SignalProcessing/Private/FloatArrayMath.cpp b/Engine/Source/Runtime/SignalProcessing/Private/FloatArrayMath.cpp
--- a/Engine/Source/Runtime/SignalProcessing/Private/FloatArrayMath.cpp
+++ b/Engine/Source/Runtime/SignalProcessing/Private/FloatArrayMath.cpp
@@ -2,6 +2,20 @@
 
 #include "DSP/FloatArrayMath.h"
 #include "CoreMinimal.h"
+#include "SignalProcessingModule.h"
+
+#if INTEL_ISPC && !UE_BUILD_SHIPPING
+#include "HAL/IConsoleManager.h"
+#endif
+
+#if INTEL_ISPC
+#include "FloatArrayMath.ispc.generated.h"
+#endif
+
+#if INTEL_ISPC && !UE_BUILD_SHIPPING
+bool bAudio_FloatArrayMath_ISPC_Enabled = true;
+FAutoConsoleVariableRef CVarAudioFloatArrayMathISPCEnabled(TEXT("au.FloatArrayMath.ISPC"), bAudio_FloatArrayMath_ISPC_Enabled, TEXT("Whether to use ISPC optimizations in audio float array math operations"));
+#endif
 
 namespace Audio
 {
@@ -17,49 +31,101 @@ namespace Audio
 		OutSum = 0.f;
 
 		int32 Num = InValues.Num();
-		const float* InData = InValues.GetData();
 
-		for (int32 i = 0; i < Num; i++)
+		if (bAudio_FloatArrayMath_ISPC_Enabled)
 		{
-			OutSum += InData[i];
+#if INTEL_ISPC
+			ispc::ArraySum(InValues.GetData(), OutSum, Num);
+#endif
+		}
+		else
+		{
+			const int32 NumToSimd = Num & MathIntrinsics::SimdMask;
+			const int32 NumNotToSimd = Num & MathIntrinsics::NotSimdMask;
+
+			if (NumToSimd)
+			{
+				VectorRegister4Float Total = VectorSetFloat1(0.f);
+
+				for (int32 i = 0; i < NumToSimd; i += AUDIO_NUM_FLOATS_PER_VECTOR_REGISTER)
+				{
+					VectorRegister4Float VectorData = VectorLoad(&InValues[i]);
+					Total = VectorAdd(Total, VectorData);
+				}
+
+				float Val[4];
+				VectorStore(Total, Val);
+				OutSum += Val[0] + Val[1] + Val[2] + Val[3];
+			}
+
+			if (NumNotToSimd)
+			{
+				for (int32 i = NumToSimd; i < Num; ++i)
+				{
+					OutSum += InValues[i];
+				}
+			}
+
 		}
 	}
 
 	void ArraySum(const FAlignedFloatBuffer& InValues, float& OutSum)
 	{
-		OutSum = 0.f;
+		TArrayView<const float> InValuesView(InValues.GetData(), InValues.Num());
 
-		const int32 Num = InValues.Num();
-		const int32 NumToSimd = Num & MathIntrinsics::SimdMask;
-		const int32 NumNotToSimd = Num & MathIntrinsics::NotSimdMask;
+		ArraySum(InValuesView, OutSum);
+	}
 
-		const float* InData = InValues.GetData();
+	void ArraySum(TArrayView<const float> InFloatBuffer1, TArrayView<const float> InFloatBuffer2, TArrayView<float> OutputBuffer)
+	{
+		checkf(InFloatBuffer1.Num() == InFloatBuffer2.Num(), TEXT("Input buffers must be equal length"));
+
+		const int32 Num = InFloatBuffer1.Num();
 
-		if (NumToSimd)
+		if (bAudio_FloatArrayMath_ISPC_Enabled)
+		{
+#if INTEL_ISPC
+			ispc::ArraySum2(InFloatBuffer1.GetData(), InFloatBuffer2.GetData(), OutputBuffer.GetData(), Num);
+#endif
+		}
+		else
 		{
-			VectorRegister4Float Total = VectorSetFloat1(0.f);
+			const int32 NumToSimd = Num & MathIntrinsics::SimdMask;
+			const int32 NumNotToSimd = Num & MathIntrinsics::NotSimdMask;
 
-			for (int32 i = 0; i < NumToSimd; i += 4)
+			if (NumToSimd)
 			{
-				VectorRegister4Float VectorData = VectorLoadAligned(&InData[i]);
-				Total = VectorAdd(Total, VectorData);
+				for (int32 i = 0; i < NumToSimd; i += AUDIO_NUM_FLOATS_PER_VECTOR_REGISTER)
+				{
+					VectorRegister4Float Input1 = VectorLoad(&InFloatBuffer1[i]);
+					VectorRegister4Float Input2 = VectorLoad(&InFloatBuffer2[i]);
+
+					VectorRegister4Float Output = VectorAdd(Input1, Input2);
+					VectorStore(Output, &OutputBuffer[i]);
+				}
 			}
 
-			MS_ALIGN(16) float Val[4] GCC_ALIGN(16);
-			VectorStoreAligned(Total, Val);
-			OutSum = Val[0] + Val[1] + Val[2] + Val[3];
+			if (NumNotToSimd)
+			{
+				for (int32 i = NumToSimd; i < Num; ++i)
+				{
+					OutputBuffer[i] = InFloatBuffer1[i] + InFloatBuffer2[i];
+				}
+			}
 		}
+	}
 
-		if (NumNotToSimd)
-		{
-			TArrayView<const float> ValuesView(&InData[NumToSimd],  NumNotToSimd);
-
-			float ExtraSum = 0.f;
+	void ArraySum(const FAlignedFloatBuffer& InFloatBuffer1, const FAlignedFloatBuffer& InFloatBuffer2, FAlignedFloatBuffer& OutputBuffer)
+	{
+		const int32 InNum = InFloatBuffer1.Num();
+		OutputBuffer.Reset(InNum);
+		OutputBuffer.AddUninitialized(InNum);
 
-			ArraySum(ValuesView, ExtraSum);
+		TArrayView<const float> InFloatBuffer1View(InFloatBuffer1.GetData(), InNum);
+		TArrayView<const float> InFloatBuffer2View(InFloatBuffer2.GetData(), InNum);
+		TArrayView<float> OutputBufferView(OutputBuffer.GetData(), InNum);
 
-			OutSum += ExtraSum;
-		}
+		ArraySum(InFloatBuffer1View, InFloatBuffer2View, OutputBufferView);
 	}
 
 	void ArrayCumulativeSum(TArrayView<const float> InView, TArray<float>& OutData)
@@ -77,13 +143,22 @@ namespace Audio
 		float* OutDataPtr = OutData.GetData();
 		const float* InViewPtr = InView.GetData();
 
-		// Start summing
-		*OutDataPtr = *InViewPtr++;
-
-		for (int32 i = 1; i < Num; i++)
+		if (bAudio_FloatArrayMath_ISPC_Enabled)
+		{
+#if INTEL_ISPC
+			ispc::ArrayCumulativeSum(InViewPtr, OutDataPtr, Num);
+#endif
+		}
+		else
 		{
-			float Temp = *OutDataPtr++ + *InViewPtr++;
-			*OutDataPtr = Temp;
+			// Start summing
+			*OutDataPtr = *InViewPtr++;
+	
+			for (int32 i = 1; i < Num; i++)
+			{
+				float Temp = *OutDataPtr++ + *InViewPtr++;
+				*OutDataPtr = Temp;
+			}
 		}
 	}
 
@@ -100,12 +175,21 @@ namespace Audio
 
 		const float* DataPtr = InView.GetData();
 
-		for (int32 i = 0; i < Num; i++)
+		if (bAudio_FloatArrayMath_ISPC_Enabled)
 		{
-			OutMean += DataPtr[i];
+#if INTEL_ISPC
+			ispc::ArrayMean(DataPtr, OutMean, Num);
+#endif
+		}
+		else
+		{
+			for (int32 i = 0; i < Num; i++)
+			{
+				OutMean += DataPtr[i];
+			}
+	
+			OutMean /= static_cast<float>(Num);
 		}
-
-		OutMean /= static_cast<float>(Num);
 	}
 
 	void ArrayMeanSquared(TArrayView<const float> InView, float& OutMean)
@@ -121,12 +205,177 @@ namespace Audio
 
 		const float* DataPtr = InView.GetData();
 
-		for (int32 i = 0; i < Num; i++)
+		if (bAudio_FloatArrayMath_ISPC_Enabled)
+		{
+#if INTEL_ISPC
+			ispc::ArrayMeanSquared(DataPtr, OutMean, Num);
+#endif
+		}
+		else
+		{
+			for (int32 i = 0; i < Num; i++)
+			{
+				OutMean += DataPtr[i] * DataPtr[i];
+			}
+	
+			OutMean /= static_cast<float>(Num);
+		}
+	}
+
+	float ArrayGetMagnitude(TArrayView<const float> Buffer)
+	{
+		const int32 Num = Buffer.Num();
+
+		if (bAudio_FloatArrayMath_ISPC_Enabled)
+		{
+#if INTEL_ISPC
+			return ispc::ArrayGetMagnitude(Buffer.GetData(), Num);
+#endif
+		}
+		else
+		{
+			const int32 NumToSimd = Num & MathIntrinsics::SimdMask;
+			const int32 NumNotToSimd = Num & MathIntrinsics::NotSimdMask;
+
+			float Sum = 0.0f;
+
+			if (NumToSimd)
+			{
+				VectorRegister4Float VectorSum = VectorZero();
+
+				const float Exponent = 2.0f;
+				VectorRegister4Float ExponentVector = VectorLoadFloat1(&Exponent);
+
+				for (int32 i = 0; i < NumToSimd; i += AUDIO_NUM_FLOATS_PER_VECTOR_REGISTER)
+				{
+					VectorRegister4Float Input = VectorPow(VectorLoad(&Buffer[i]), ExponentVector);
+					VectorSum = VectorAdd(VectorSum, Input);
+				}
+
+				float PartionedSums[4];
+				VectorStore(VectorSum, PartionedSums);
+
+				Sum += PartionedSums[0] + PartionedSums[1] + PartionedSums[2] + PartionedSums[3];
+			}
+
+			if (NumNotToSimd)
+			{
+				for (int32 i = NumToSimd; i < Num; ++i)
+				{
+					Sum += Buffer[i] * Buffer[i];
+				}
+			}
+
+			return FMath::Sqrt(Sum);
+		}
+	}
+
+	float ArrayGetMagnitude(const FAlignedFloatBuffer& Buffer)
+	{
+		TArrayView<const float> BufferView(Buffer.GetData(), Buffer.Num());
+		return ArrayGetMagnitude(BufferView);
+	}
+
+	float ArrayGetAverageValue(TArrayView<const float> Buffer)
+	{
+		const int32 Num = Buffer.Num();
+
+		if (bAudio_FloatArrayMath_ISPC_Enabled)
+		{
+#if INTEL_ISPC
+			return ispc::ArrayGetAverageValue(Buffer.GetData(), Num);
+#endif
+		}
+		else
+		{
+			const int32 NumToSimd = Num & MathIntrinsics::SimdMask;
+			const int32 NumNotToSimd = Num & MathIntrinsics::NotSimdMask;
+
+			float Sum = 0.0f;
+
+			if (NumToSimd)
+			{
+				VectorRegister4Float VectorSum = VectorZero();
+
+				for (int32 i = 0; i < NumToSimd; i += AUDIO_NUM_FLOATS_PER_VECTOR_REGISTER)
+				{
+					VectorRegister4Float Input = VectorLoad(&Buffer[i]);
+					VectorSum = VectorAdd(VectorSum, Input);
+				}
+
+				float PartionedSums[4];
+				VectorStore(VectorSum, PartionedSums);
+
+				Sum += PartionedSums[0] + PartionedSums[1] + PartionedSums[2] + PartionedSums[3];
+			}
+
+			if (NumNotToSimd)
+			{
+				for (int32 i = NumToSimd; i < Num; ++i)
+				{
+					Sum += Buffer[i];
+				}
+			}
+
+			return Sum / Num;
+		}
+	}
+
+	float ArrayGetAverageValue(const FAlignedFloatBuffer& Buffer)
+	{
+		TArrayView<const float> BufferView(Buffer.GetData(), Buffer.Num());
+		return ArrayGetAverageValue(BufferView);
+	}
+
+	float ArrayGetAverageAbsValue(const FAlignedFloatBuffer& Buffer)
+	{
+		TArrayView<const float> BufferView(Buffer.GetData(), Buffer.Num());
+		return ArrayGetAverageAbsValue(BufferView);
+	}
+
+	float ArrayGetAverageAbsValue(TArrayView<const float> Buffer)
+	{
+		const int32 Num = Buffer.Num();
+
+		if (bAudio_FloatArrayMath_ISPC_Enabled)
 		{
-			OutMean += DataPtr[i] * DataPtr[i];
+#if INTEL_ISPC
+			return ispc::ArrayGetAverageAbsValue(Buffer.GetData(), Num);
+#endif
 		}
+		else
+		{
+			const int32 NumToSimd = Num & MathIntrinsics::SimdMask;
+			const int32 NumNotToSimd = Num & MathIntrinsics::NotSimdMask;
+
+			float Sum = 0.0f;
+
+			if (NumToSimd)
+			{
+				VectorRegister4Float VectorSum = VectorZero();
 
-		OutMean /= static_cast<float>(Num);
+				for (int32 i = 0; i < NumToSimd; i += AUDIO_NUM_FLOATS_PER_VECTOR_REGISTER)
+				{
+					VectorRegister4Float Input = VectorAbs(VectorLoad(&Buffer[i]));
+					VectorSum = VectorAdd(VectorSum, Input);
+				}
+
+				float PartionedSums[4];
+				VectorStore(VectorSum, PartionedSums);
+
+				Sum += PartionedSums[0] + PartionedSums[1] + PartionedSums[2] + PartionedSums[3];
+			}
+
+			if (NumNotToSimd)
+			{
+				for (int32 i = NumToSimd; i < Num; ++i)
+				{
+					Sum += FMath::Abs(Buffer[i]);
+				}
+			}
+
+			return Sum / Num;
+		}
 	}
 
 	void ArrayMeanFilter(TArrayView<const float> InView, int32 WindowSize, int32 WindowOrigin, TArray<float>& OutData)
@@ -153,46 +402,54 @@ namespace Audio
 		TArray<float> SummedData;
 		ArrayCumulativeSum(InView, SummedData);
 		const float LastSummedData = SummedData.Last();
-		
-		
-		const int32 LastIndexBeforeEndBoundaryCondition = FMath::Max(WindowOrigin + 1, Num - WindowSize + WindowOrigin + 1);
-		const int32 StartOffset = -WindowOrigin - 1;
-		const int32 EndOffset = WindowSize - WindowOrigin - 1;
-		const int32 WindowTail = WindowSize - WindowOrigin;
 
 		float* OutDataPtr = OutData.GetData();
 		const float* SummedDataPtr = SummedData.GetData();
 
-		if ((WindowSize - WindowOrigin) < Num)
+		if (bAudio_FloatArrayMath_ISPC_Enabled)
+		{
+#if INTEL_ISPC
+			ispc::ArrayMeanFilter(SummedDataPtr, WindowSize, WindowOrigin, OutDataPtr, LastSummedData, Num);
+#endif
+		}
+		else
 		{
-			// Handle boundary condition where analysis window precedes beginning of array.
-			for (int32 i = 0; i < (WindowOrigin + 1); i++)
+			const int32 LastIndexBeforeEndBoundaryCondition = FMath::Max(WindowOrigin + 1, Num - WindowSize + WindowOrigin + 1);
+			const int32 StartOffset = -WindowOrigin - 1;
+			const int32 EndOffset = WindowSize - WindowOrigin - 1;
+			const int32 WindowTail = WindowSize - WindowOrigin;
+	
+			if ((WindowSize - WindowOrigin) < Num)
 			{
-				OutDataPtr[i] = SummedDataPtr[i + EndOffset] / FMath::Max(1.f, static_cast<float>(WindowTail + i));
+				// Handle boundary condition where analysis window precedes beginning of array.
+				for (int32 i = 0; i < (WindowOrigin + 1); i++)
+				{
+					OutDataPtr[i] = SummedDataPtr[i + EndOffset] / FMath::Max(1.f, static_cast<float>(WindowTail + i));
+				}
+	
+				// No boundary conditions to handle here.	
+				const float MeanDivisor = static_cast<float>(WindowSize);
+				for (int32 i = WindowOrigin + 1; i < LastIndexBeforeEndBoundaryCondition; i++)
+				{
+					OutDataPtr[i] = (SummedDataPtr[i + EndOffset] - SummedDataPtr[i + StartOffset]) / MeanDivisor;
+				}
 			}
-
-			// No boundary conditions to handle here.	
-			const float MeanDivisor = static_cast<float>(WindowSize);
-			for (int32 i = WindowOrigin + 1; i < LastIndexBeforeEndBoundaryCondition; i++)
+			else
 			{
-				OutDataPtr[i] = (SummedDataPtr[i + EndOffset] - SummedDataPtr[i + StartOffset]) / MeanDivisor;
+				// Handle boundary condition where window precedes beginning and goes past end of array
+				const float ArrayMean = LastSummedData / static_cast<float>(Num);
+				for (int32 i = 0; i < LastIndexBeforeEndBoundaryCondition; i++)
+				{
+					OutDataPtr[i] = ArrayMean;
+				}
 			}
-		}
-		else
-		{
-			// Handle boundary condition where window precedes beginning and goes past end of array
-			const float ArrayMean = LastSummedData / static_cast<float>(Num);
-			for (int32 i = 0; i < LastIndexBeforeEndBoundaryCondition; i++)
+	
+			// Handle boundary condition where analysis window goes past end of array.
+			for (int32 i = LastIndexBeforeEndBoundaryCondition; i < Num; i++)
 			{
-				OutDataPtr[i] = ArrayMean;
+				OutDataPtr[i] = (LastSummedData - SummedDataPtr[i + StartOffset]) / static_cast<float>(Num - i + WindowOrigin);
 			}
 		}
-
-		// Handle boundary condition where analysis window goes past end of array.
-		for (int32 i = LastIndexBeforeEndBoundaryCondition; i < Num; i++)
-		{
-			OutDataPtr[i] = (LastSummedData - SummedDataPtr[i + StartOffset]) / static_cast<float>(Num - i + WindowOrigin);
-		}
 	}
 
 	void ArrayMaxFilter(TArrayView<const float> InView, int32 WindowSize, int32 WindowOrigin, TArray<float>& OutData)
@@ -280,14 +537,23 @@ namespace Audio
 		OutEuclideanNorm = 0.0f;
 		const int32 Num = InView.Num();
 		const float* InViewData = InView.GetData();
-		
-		// Sum it up.
-		for (int32 i = 0; i < Num; i++)
+
+		if (bAudio_FloatArrayMath_ISPC_Enabled)
 		{
-			OutEuclideanNorm += InViewData[i] * InViewData[i];
+#if INTEL_ISPC
+			ispc::ArrayGetEuclideanNorm(InViewData, OutEuclideanNorm, Num);
+#endif
+		}
+		else
+		{
+			// Sum it up.
+			for (int32 i = 0; i < Num; i++)
+			{
+				OutEuclideanNorm += InViewData[i] * InViewData[i];
+			}
+	
+			OutEuclideanNorm = FMath::Sqrt(OutEuclideanNorm);
 		}
-
-		OutEuclideanNorm = FMath::Sqrt(OutEuclideanNorm);
 	}
 
 	void ArrayAbs(TArrayView<const float> InBuffer, TArrayView<float> OutBuffer)
@@ -298,9 +564,18 @@ namespace Audio
 		const float* InData = InBuffer.GetData();
 		float* OutData = OutBuffer.GetData();
 
-		for (int32 i = 0; i < Num; i++)
+		if (bAudio_FloatArrayMath_ISPC_Enabled)
+		{
+#if INTEL_ISPC
+			ispc::ArrayAbs(InData, OutData, Num);
+#endif
+		}
+		else
 		{
-			OutData[i] = FMath::Abs(InData[i]);
+			for (int32 i = 0; i < Num; i++)
+			{
+				OutData[i] = FMath::Abs(InData[i]);
+			}
 		}
 	}
 
@@ -310,9 +585,18 @@ namespace Audio
 		const int32 Num = InView.Num();
 		float* Data = InView.GetData();
 
-		for (int32 i = 0; i < Num; i++)
+		if (bAudio_FloatArrayMath_ISPC_Enabled)
 		{
-			Data[i] = FMath::Abs(Data[i]);
+#if INTEL_ISPC
+			ispc::ArrayAbsInPlace(Data, Num);
+#endif
+		}
+		else
+		{
+			for (int32 i = 0; i < Num; i++)
+			{
+				Data[i] = FMath::Abs(Data[i]);
+			}
 		}
 	}
 
@@ -321,9 +605,18 @@ namespace Audio
 		const int32 Num = InView.Num();
 		float* Data = InView.GetData();
 
-		for (int32 i = 0; i < Num; i++)
+		if (bAudio_FloatArrayMath_ISPC_Enabled)
+		{
+#if INTEL_ISPC
+			ispc::ArrayClampMinInPlace(Data, InMin, Num);
+#endif
+		}
+		else
 		{
-			Data[i] = FMath::Max(InMin, Data[i]);
+			for (int32 i = 0; i < Num; i++)
+			{
+				Data[i] = FMath::Max(InMin, Data[i]);
+			}
 		}
 	}
 
@@ -332,9 +625,18 @@ namespace Audio
 		const int32 Num = InView.Num();
 		float* Data = InView.GetData();
 
-		for (int32 i = 0; i < Num; i++)
+		if (bAudio_FloatArrayMath_ISPC_Enabled)
+		{
+#if INTEL_ISPC
+			ispc::ArrayClampMaxInPlace(Data, InMax, Num);
+#endif
+		}
+		else
 		{
-			Data[i] = FMath::Min(InMax, Data[i]);
+			for (int32 i = 0; i < Num; i++)
+			{
+				Data[i] = FMath::Min(InMax, Data[i]);
+			}
 		}
 	}
 
@@ -343,9 +645,18 @@ namespace Audio
 		const int32 Num = InView.Num();
 		float* Data = InView.GetData();
 
-		for (int32 i = 0; i < Num; i++)
+		if (bAudio_FloatArrayMath_ISPC_Enabled)
+		{
+#if INTEL_ISPC
+			ispc::ArrayClampInPlace(Data, InMin, InMax, Num);
+#endif
+		}
+		else
 		{
-			Data[i] = FMath::Clamp(Data[i], InMin, InMax);
+			for (int32 i = 0; i < Num; i++)
+			{
+				Data[i] = FMath::Clamp(Data[i], InMin, InMax);
+			}
 		}
 	}
 
@@ -362,69 +673,87 @@ namespace Audio
 		OutArray.AddUninitialized(Num);
 
 		const float* InDataPtr = InView.GetData();
-		float MaxValue = InDataPtr[0];
-		float MinValue = InDataPtr[0];
+		float* OutDataPtr = OutArray.GetData();
 
-		// determine min and max
-		for (int32 i = 1; i < Num; i++)
+		if (bAudio_FloatArrayMath_ISPC_Enabled)
+		{
+#if INTEL_ISPC
+			ispc::ArrayMinMaxNormalize(InDataPtr, OutDataPtr, Num);
+#endif
+		}
+		else
 		{
-			if (InDataPtr[i] < MinValue)
+			float MaxValue = InDataPtr[0];
+			float MinValue = InDataPtr[0];
+	
+			// determine min and max
+			for (int32 i = 1; i < Num; i++)
 			{
-				MinValue = InDataPtr[i];
+				if (InDataPtr[i] < MinValue)
+				{
+					MinValue = InDataPtr[i];
+				}
+				else if (InDataPtr[i] > MaxValue)
+				{
+					MaxValue = InDataPtr[i];
+				}
 			}
-			else if (InDataPtr[i] > MaxValue)
+	
+			// Normalize data by subtracting minimum value and dividing by range
+			float Scale = 1.f / FMath::Max(SMALL_NUMBER, MaxValue - MinValue);
+			for (int32 i = 0; i < Num; i++)
 			{
-				MaxValue = InDataPtr[i];
+				OutDataPtr[i] = (InDataPtr[i] - MinValue) * Scale;
 			}
 		}
-
-		// Normalize data by subtracting minimum value and dividing by range
-		float* OutDataPtr = OutArray.GetData();
-		float Scale = 1.f / FMath::Max(SMALL_NUMBER, MaxValue - MinValue);
-		for (int32 i = 0; i < Num; i++)
-		{
-			OutDataPtr[i] = (InDataPtr[i] - MinValue) * Scale;
-		}
 	}
 
-	void ArrayMultiplyInPlace(TArrayView<const float> InValues1, TArrayView<float> InValues2)
+	void ArrayMultiplyInPlace(TArrayView<const float> InFloatBuffer, TArrayView<float> BufferToMultiply)
 	{
-		check(InValues1.Num() == InValues2.Num());
-
-		const int32 Num = InValues1.Num();
+		checkf(InFloatBuffer.Num() == BufferToMultiply.Num(), TEXT("Input buffers must be equal length"));
 
-		const float* InData1 = InValues1.GetData();
-		float* InData2 = InValues2.GetData();
+		const int32 Num = BufferToMultiply.Num();
 
-		for (int32 i = 0; i < Num; i++)
+		if (bAudio_FloatArrayMath_ISPC_Enabled)
 		{
-			InData2[i] *= InData1[i];
+#if INTEL_ISPC
+			ispc::ArrayMultiplyInPlace(InFloatBuffer.GetData(), BufferToMultiply.GetData(), Num);
+#endif
 		}
-	}
-
-	void ArrayMultiplyInPlace(const FAlignedFloatBuffer& InValues1, FAlignedFloatBuffer& InValues2)
-	{
-		check(InValues1.Num() == InValues2.Num());
+		else
+		{
+			const int32 NumToSimd = Num & MathIntrinsics::SimdMask;
+			const int32 NumNotToSimd = Num & MathIntrinsics::NotSimdMask;
 
-		const int32 Num = InValues1.Num();
-		const int32 NumToSimd = Num & MathIntrinsics::SimdMask;
-		const int32 NumNotToSimd = Num & MathIntrinsics::NotSimdMask;
+			if (NumToSimd)
+			{
+				for (int32 i = 0; i < NumToSimd; i += AUDIO_NUM_FLOATS_PER_VECTOR_REGISTER)
+				{
+					VectorRegister4Float Input1 = VectorLoad(&InFloatBuffer[i]);
+					VectorRegister4Float Output = VectorLoad(&BufferToMultiply[i]);
 
-		const float* InData1 = InValues1.GetData();
-		float* InData2 = InValues2.GetData();
+					Output = VectorMultiply(Input1, Output);
+					VectorStore(Output, &BufferToMultiply[i]);
+				}
+			}
 
-		if (NumToSimd)
-		{
-			MultiplyBuffersInPlace(InData1, InData2, NumToSimd);
+			if (NumNotToSimd)
+			{
+				for (int32 i = NumToSimd; i < Num; ++i)
+				{
+					BufferToMultiply[i] = InFloatBuffer[i] * BufferToMultiply[i];
+				}
+			}
 		}
+	}
 
-		if (NumNotToSimd)
-		{
-			TArrayView<const float> ValuesView1(&InData1[NumToSimd],  NumNotToSimd);
-			TArrayView<float> ValuesView2(&InData2[NumToSimd],  NumNotToSimd);
+	void ArrayMultiplyInPlace(const FAlignedFloatBuffer& InFloatBuffer, FAlignedFloatBuffer& BufferToMultiply)
+	{
+		const int32 InNum = BufferToMultiply.Num();
+		TArrayView<const float> InFloatBufferView(InFloatBuffer.GetData(), InNum);
+		TArrayView<float> BufferToMultiplyView(BufferToMultiply.GetData(), InNum);
 
-			ArrayMultiplyInPlace(ValuesView1, ValuesView2);
-		}
+		ArrayMultiplyInPlace(InFloatBufferView, BufferToMultiplyView);
 	}
 
 	void ArrayComplexMultiplyInPlace(TArrayView<const float> InValues1, TArrayView<float> InValues2)
@@ -439,88 +768,162 @@ namespace Audio
 		const float* InData1 = InValues1.GetData();
 		float* InData2 = InValues2.GetData();
 
-		for (int32 i = 0; i < Num; i += 2)
+		if (bAudio_FloatArrayMath_ISPC_Enabled)
 		{
-			float Real = (InData1[i] * InData2[i]) - (InData1[i + 1] * InData2[i + 1]);
-			float Imag = (InData1[i] * InData2[i + 1]) + (InData1[i + 1] * InData2[i]);
-			InData2[i] = Real;
-			InData2[i + 1] = Imag;
+#if INTEL_ISPC
+			ispc::ArrayComplexMultiplyInPlace(InData1, InData2, Num);
+#endif
 		}
-	}
-
-	void ArrayComplexMultiplyInPlace(const FAlignedFloatBuffer& InValues1, FAlignedFloatBuffer& InValues2)
-	{
-		check(InValues1.Num() == InValues2.Num());
-
-		const int32 Num = InValues1.Num();
-		const int32 NumToSimd = Num & MathIntrinsics::SimdMask;
-		const int32 NumNotToSimd = Num & MathIntrinsics::NotSimdMask;
-
-		const float* InData1 = InValues1.GetData();
-		float* InData2 = InValues2.GetData();
-
-		if (NumToSimd)
+		else
 		{
-			const VectorRegister4Float RealSignFlip = MakeVectorRegister(-1.f, 1.f, -1.f, 1.f);
+			const int32 NumToSimd = Num & MathIntrinsics::SimdMask;
+			const int32 NumNotToSimd = Num & MathIntrinsics::NotSimdMask;
 
-			for (int32 i = 0; i < NumToSimd; i += 4)
+			if (NumToSimd)
 			{
-				VectorRegister4Float VectorData1 = VectorLoadAligned(&InData1[i]);
-				VectorRegister4Float VectorData2 = VectorLoadAligned(&InData2[i]);
+				const VectorRegister4Float RealSignFlip = MakeVectorRegister(-1.f, 1.f, -1.f, 1.f);
 
-				VectorRegister4Float VectorData1Real = VectorSwizzle(VectorData1, 0, 0, 2, 2);
-				VectorRegister4Float VectorData1Imag = VectorSwizzle(VectorData1, 1, 1, 3, 3);
-				VectorRegister4Float VectorData2Swizzle = VectorSwizzle(VectorData2, 1, 0, 3, 2);
+				for (int32 i = 0; i < NumToSimd; i += AUDIO_NUM_FLOATS_PER_VECTOR_REGISTER)
+				{
+					VectorRegister4Float VectorData1 = VectorLoad(&InData1[i]);
+					VectorRegister4Float VectorData2 = VectorLoad(&InData2[i]);
 
-				VectorRegister4Float Result = VectorMultiply(VectorData1Imag, VectorData2Swizzle);
-				Result = VectorMultiply(Result, RealSignFlip);
-				Result = VectorMultiplyAdd(VectorData1Real, VectorData2, Result);
+					VectorRegister4Float VectorData1Real = VectorSwizzle(VectorData1, 0, 0, 2, 2);
+					VectorRegister4Float VectorData1Imag = VectorSwizzle(VectorData1, 1, 1, 3, 3);
+					VectorRegister4Float VectorData2Swizzle = VectorSwizzle(VectorData2, 1, 0, 3, 2);
 
-				VectorStoreAligned(Result, &InData2[i]);
-			}
-		}
+					VectorRegister4Float Result = VectorMultiply(VectorData1Imag, VectorData2Swizzle);
+					Result = VectorMultiply(Result, RealSignFlip);
+					Result = VectorMultiplyAdd(VectorData1Real, VectorData2, Result);
 
-		if (NumNotToSimd)
-		{
-			TArrayView<const float> ValuesView1(&InData1[NumToSimd],  NumNotToSimd);
-			TArrayView<float> ValuesView2(&InData2[NumToSimd],  NumNotToSimd);
+					VectorStore(Result, &InData2[i]);
+				}
+			}
 
-			ArrayComplexMultiplyInPlace(ValuesView1, ValuesView2);
+			if (NumNotToSimd)
+			{
+				for (int32 i = NumToSimd; i < Num; i += 2)
+				{
+					float Real = (InData1[i] * InData2[i]) - (InData1[i + 1] * InData2[i + 1]);
+					float Imag = (InData1[i] * InData2[i + 1]) + (InData1[i + 1] * InData2[i]);
+					InData2[i] = Real;
+					InData2[i + 1] = Imag;
+				}
+			}
 		}
 	}
 
-	void ArrayMultiplyByConstantInPlace(TArrayView<float> InValues, float InMultiplier)
+	void ArrayComplexMultiplyInPlace(const FAlignedFloatBuffer& InValues1, FAlignedFloatBuffer& InValues2)
 	{
-		const int32 Num = InValues.Num();
-		float* InData = InValues.GetData();
+		TArrayView<const float> Values1View(InValues1.GetData(), InValues1.Num());
+		TArrayView<float> Values2View(InValues2.GetData(), InValues2.Num());
+
+		ArrayComplexMultiplyInPlace(Values1View, Values2View);
+	}
+
+	void ArrayMultiplyByConstant(TArrayView<const float> InFloatBuffer, float InValue, TArrayView<float> OutFloatBuffer)
+	{
+		check(InFloatBuffer.Num() == OutFloatBuffer.Num());
+
+		const int32 Num = InFloatBuffer.Num();
 
-		for (int32 i = 0; i < Num; i++)
+		// Get ptrs to audio buffers to avoid bounds check in non-shipping builds
+		const float* InBufferPtr = InFloatBuffer.GetData();
+		float* OutBufferPtr = OutFloatBuffer.GetData();
+
+		if (bAudio_FloatArrayMath_ISPC_Enabled)
 		{
-			InData[i] *= InMultiplier;
+#if INTEL_ISPC
+			ispc::ArrayMultiplyByConstant(InBufferPtr, InValue, OutBufferPtr, Num);
+#endif
+		}
+		else
+		{
+			// Can only SIMD on multiple of 4 buffers, we'll do normal multiples on last bit
+			const int32 NumToSimd = Num & MathIntrinsics::SimdMask;
+			const int32 NumNotToSimd = Num & MathIntrinsics::NotSimdMask;
+
+			if (NumToSimd)
+			{
+				// Load the single value we want to multiply all values by into a vector register
+				const VectorRegister4Float MultiplyValue = VectorLoadFloat1(&InValue);
+				for (int32 i = 0; i < NumToSimd; i += AUDIO_NUM_FLOATS_PER_VECTOR_REGISTER)
+				{
+					// Load the next 4 samples of the input buffer into a register
+					VectorRegister4Float InputBufferRegister = VectorLoad(&InBufferPtr[i]);
+
+					// Perform the multiply
+					VectorRegister4Float Temp = VectorMultiply(InputBufferRegister, MultiplyValue);
+
+					// Store results into the output buffer
+					VectorStore(Temp, &OutBufferPtr[i]);
+				}
+			}
+
+			if (NumNotToSimd)
+			{
+				// Perform remaining non-simd values left over
+				for (int32 i = NumToSimd; i < Num; ++i)
+				{
+					OutBufferPtr[i] = InValue * InBufferPtr[i];
+				}
+			}
 		}
 	}
 
-	void ArrayMultiplyByConstantInPlace(FAlignedFloatBuffer& InValues, float InMultiplier)
+	void ArrayMultiplyByConstant(const FAlignedFloatBuffer& InFloatBuffer, float InValue, FAlignedFloatBuffer& OutFloatBuffer)
 	{
-		const int32 Num = InValues.Num();
-		const int32 NumToSimd = Num & MathIntrinsics::SimdMask;
-		const int32 NumNotToSimd = Num & MathIntrinsics::NotSimdMask;
+		TArrayView<const float> InFloatBufferView(InFloatBuffer.GetData(), InFloatBuffer.Num());
+		TArrayView<float> OutFloatBufferView(OutFloatBuffer.GetData(), OutFloatBuffer.Num());
 
-		float* InData = InValues.GetData();
+		ArrayMultiplyByConstant(InFloatBufferView, InValue, OutFloatBufferView);
+	}
 
-		if (NumToSimd)
+	void ArrayMultiplyByConstantInPlace(TArrayView<float> InOutBuffer, float InGain)
+	{
+		int32 Num = InOutBuffer.Num();
+		float* InOutData = InOutBuffer.GetData();
+
+		if (bAudio_FloatArrayMath_ISPC_Enabled)
 		{
-			MultiplyBufferByConstantInPlace(InData, NumToSimd, InMultiplier);
+#if INTEL_ISPC
+			ispc::ArrayMultiplyByConstantInPlace(InOutData, Num, InGain);
+#endif
 		}
-
-		if (NumNotToSimd)
+		else
 		{
-			TArrayView<float> ValuesView(&InData[NumToSimd],  NumNotToSimd);
+			const int32 NumToSimd = Num & MathIntrinsics::SimdMask;
+			const int32 NumNotToSimd = Num & MathIntrinsics::NotSimdMask;
 
-			ArrayMultiplyByConstantInPlace(ValuesView, InMultiplier);
+			const VectorRegister4Float Gain = VectorLoadFloat1(&InGain);
+
+			if (NumToSimd)
+			{
+				for (int32 i = 0; i < NumToSimd; i += AUDIO_NUM_FLOATS_PER_VECTOR_REGISTER)
+				{
+					VectorRegister4Float Output = VectorLoad(&InOutData[i]);
+					Output = VectorMultiply(Output, Gain);
+					VectorStore(Output, &InOutData[i]);
+				}
+			}
+
+			if (NumNotToSimd)
+			{
+				for (int32 i = NumToSimd; i < Num; i++)
+				{
+					InOutData[i] *= InGain;
+				}
+			}
 		}
 	}
 
+	void ArrayMultiplyByConstantInPlace(FAlignedFloatBuffer& InOutBuffer, float InGain)
+	{
+		TArrayView<float> InOutBufferView(InOutBuffer.GetData(), InOutBuffer.Num());
+
+		ArrayMultiplyByConstantInPlace(InOutBufferView, InGain);
+	}
+
 	void ArrayAddInPlace(TArrayView<const float> InValues, TArrayView<float> InAccumulateValues)
 	{
 		check(InValues.Num() == InAccumulateValues.Num());
@@ -530,41 +933,93 @@ namespace Audio
 		const float* InData = InValues.GetData();
 		float* InAccumulateData = InAccumulateValues.GetData();
 
-		for (int32 i = 0; i < Num; i++)
+		if (bAudio_FloatArrayMath_ISPC_Enabled)
 		{
-			InAccumulateData[i] += InData[i];
+#if INTEL_ISPC
+			ispc::ArrayAddInPlace(InData, InAccumulateData, Num);
+#endif
+		}
+		else
+		{
+			const int32 NumToSimd = Num & MathIntrinsics::SimdMask;
+			const int32 NumNotToSimd = Num & MathIntrinsics::NotSimdMask;
+
+			if (NumToSimd)
+			{
+				for (int32 i = 0; i < NumToSimd; i += AUDIO_NUM_FLOATS_PER_VECTOR_REGISTER)
+				{
+					VectorRegister4Float VectorData = VectorLoad(&InData[i]);
+					VectorRegister4Float VectorAccumData = VectorLoad(&InAccumulateData[i]);
+
+					VectorRegister4Float VectorOut = VectorAdd(VectorData, VectorAccumData);
+					VectorStore(VectorOut, &InAccumulateData[i]);
+				}
+			}
+
+			if (NumNotToSimd)
+			{
+				for (int32 i = NumToSimd; i < Num; i++)
+				{
+					InAccumulateData[i] += InData[i];
+				}
+			}
 		}
 	}
 
 	void ArrayAddInPlace(const FAlignedFloatBuffer& InValues, FAlignedFloatBuffer& InAccumulateValues)
 	{
-		check(InValues.Num() == InAccumulateValues.Num());
+		TArrayView<const float> ValuesView(InValues.GetData(), InValues.Num());
+		TArrayView<float> AccumulateView(InAccumulateValues.GetData(), InAccumulateValues.Num());
 
-		const int32 Num = InAccumulateValues.Num();
-		const int32 NumToSimd = Num & MathIntrinsics::SimdMask;
-		const int32 NumNotToSimd = Num & MathIntrinsics::NotSimdMask;
+		ArrayAddInPlace(ValuesView, AccumulateView);
+	}
 
-		const float* InData = InValues.GetData();
-		float* InAccumulateData = InAccumulateValues.GetData();
+	void ArrayAddConstantInplace(TArrayView<float> InOutBuffer, float InConstant)
+	{
+		int32 Num = InOutBuffer.Num();
+		float* InOutData = InOutBuffer.GetData();
 
-		for (int32 i = 0; i < NumToSimd; i += 4)
+		if (bAudio_FloatArrayMath_ISPC_Enabled)
 		{
-			VectorRegister4Float VectorData = VectorLoadAligned(&InData[i]);
-			VectorRegister4Float VectorAccumData = VectorLoadAligned(&InAccumulateData[i]);
-
-			VectorRegister4Float VectorOut = VectorAdd(VectorData, VectorAccumData);
-			VectorStoreAligned(VectorOut, &InAccumulateData[i]);
+#if INTEL_ISPC
+			ispc::ArrayAddConstantInplace(InOutData, Num, InConstant);
+#endif
 		}
-
-		if (NumNotToSimd)
+		else
 		{
-			TArrayView<const float> ValuesView(&InData[NumToSimd],  NumNotToSimd);
-			TArrayView<float> AccumulateView(&InAccumulateData[NumToSimd], NumNotToSimd);
+			const int32 NumToSimd = Num & MathIntrinsics::SimdMask;
+			const int32 NumNotToSimd = Num & MathIntrinsics::NotSimdMask;
 
-			ArrayAddInPlace(ValuesView, AccumulateView);
+			const VectorRegister4Float Constant = VectorLoadFloat1(&InConstant);
+
+			if (NumToSimd)
+			{
+				for (int32 i = 0; i < NumToSimd; i += AUDIO_NUM_FLOATS_PER_VECTOR_REGISTER)
+				{
+					VectorRegister4Float Output = VectorLoad(&InOutData[i]);
+					Output = VectorAdd(Output, Constant);
+					VectorStore(Output, &InOutData[i]);
+				}
+			}
+
+			if (NumNotToSimd)
+			{
+				for (int32 i = NumToSimd; i < Num; i++)
+				{
+					InOutData[i] += InConstant;
+				}
+			}
 		}
 	}
 
+	// Adds a constant to a buffer (useful for DC offset removal)
+	void ArrayAddConstantInplace(FAlignedFloatBuffer& InOutBuffer, float InConstant)
+	{
+		TArrayView<float> InOutBufferView(InOutBuffer.GetData(), InOutBuffer.Num());
+
+		ArrayAddConstantInplace(InOutBufferView, InConstant);
+	}
+
 	void ArrayMultiplyAddInPlace(TArrayView<const float> InValues, float InMultiplier, TArrayView<float> InAccumulateValues)
 	{
 		check(InValues.Num() == InAccumulateValues.Num());
@@ -574,444 +1029,1238 @@ namespace Audio
 		const float* InData = InValues.GetData();
 		float* InAccumulateData = InAccumulateValues.GetData();
 
-		for (int32 i = 0; i < Num; i++)
+		if (bAudio_FloatArrayMath_ISPC_Enabled)
 		{
-			InAccumulateData[i] += InData[i] * InMultiplier;
+#if INTEL_ISPC
+			ispc::ArrayMultiplyAddInPlace(InData, InMultiplier, InAccumulateData, Num);
+#endif
+		}
+		else
+		{
+			for (int32 i = 0; i < Num; i++)
+			{
+				InAccumulateData[i] += InData[i] * InMultiplier;
+			}
 		}
 	}
 
 	void ArrayMultiplyAddInPlace(const FAlignedFloatBuffer& InValues, float InMultiplier, FAlignedFloatBuffer& InAccumulateValues)
+	{
+		TArrayView<const float> InValuesView(InValues.GetData(), InValues.Num());
+		TArrayView<float> InAccumulateValuesView(InAccumulateValues.GetData(), InAccumulateValues.Num());
+
+		ArrayMultiplyAddInPlace(InValuesView, InMultiplier, InAccumulateValuesView);
+	}
+
+	void ArrayLerpAddInPlace(TArrayView<const float> InValues, float InStartMultiplier, float InEndMultiplier, TArrayView<float> InAccumulateValues)
 	{
 		check(InValues.Num() == InAccumulateValues.Num());
 
-		const int32 Num = InAccumulateValues.Num();
-		const int32 NumToSimd = Num & MathIntrinsics::SimdMask;
-		const int32 NumNotToSimd = Num & MathIntrinsics::NotSimdMask;
+		const int32 Num = InValues.Num();
 
 		const float* InData = InValues.GetData();
 		float* InAccumulateData = InAccumulateValues.GetData();
 
-		MixInBufferFast(InData, InAccumulateData, NumToSimd, InMultiplier);
+		if (bAudio_FloatArrayMath_ISPC_Enabled)
+		{
+#if INTEL_ISPC
+			ispc::ArrayLerpAddInPlace(InData, InStartMultiplier, InEndMultiplier, InAccumulateData, Num);
+#endif
+		}
+		else
+		{
+			const int32 NumToSimd = Num & MathIntrinsics::SimdMask;
+			const int32 NumNotToSimd = Num & MathIntrinsics::NotSimdMask;
+
+			const float Delta = (InEndMultiplier - InStartMultiplier) / FMath::Max(1.f, static_cast<float>(Num - 1));
+
+			const float FourByDelta = 4.f * Delta;
+			VectorRegister4Float VectorDelta = MakeVectorRegister(FourByDelta, FourByDelta, FourByDelta, FourByDelta);
+			VectorRegister4Float VectorMultiplier = MakeVectorRegister(InStartMultiplier, InStartMultiplier + Delta, InStartMultiplier + 2.f * Delta, InStartMultiplier + 3.f * Delta);
+
+			if (NumToSimd)
+			{
+				for (int32 i = 0; i < NumToSimd; i += AUDIO_NUM_FLOATS_PER_VECTOR_REGISTER)
+				{
+					VectorRegister4Float VectorData = VectorLoad(&InData[i]);
+					VectorRegister4Float VectorAccumData = VectorLoad(&InAccumulateData[i]);
+
+					VectorRegister4Float VectorOut = VectorMultiplyAdd(VectorData, VectorMultiplier, VectorAccumData);
+					VectorMultiplier = VectorAdd(VectorMultiplier, VectorDelta);
+
+					VectorStore(VectorOut, &InAccumulateData[i]);
+				}
+			}
+
+			if (NumNotToSimd)
+			{
+				float Multiplier = InStartMultiplier + NumToSimd * Delta;
+
+				for (int32 i = NumToSimd; i < Num; i++)
+				{
+					InAccumulateData[i] += InData[i] * Multiplier;
+					Multiplier += Delta;
+				}
+			}
+		}
+	}
+
+	void ArrayLerpAddInPlace(const FAlignedFloatBuffer& InValues, float InStartMultiplier, float InEndMultiplier, FAlignedFloatBuffer& InAccumulateValues)
+	{
+		TArrayView<const float> ValuesView(InValues.GetData(), InValues.Num());
+		TArrayView<float> AccumulateView(InAccumulateValues.GetData(), InAccumulateValues.Num());
+
+		ArrayLerpAddInPlace(ValuesView, InStartMultiplier, InEndMultiplier, AccumulateView);
+	}
+
+	/* Subtracts two buffers together element-wise. */
+	void ArraySubtract(TArrayView<const float> InMinuend, TArrayView<const float> InSubtrahend, TArrayView<float> OutBuffer)
+	{
+		const int32 Num = InMinuend.Num();
+
+		checkf(Num == InSubtrahend.Num() && Num == OutBuffer.Num(), TEXT("InMinuend, InSubtrahend, and OutBuffer must have equal Num elements (%d vs %d vs %d)"), Num, InSubtrahend.Num(), OutBuffer.Num());
 
-		if (NumNotToSimd)
+		if (bAudio_FloatArrayMath_ISPC_Enabled)
 		{
-			TArrayView<const float> ValuesView(&InData[NumToSimd],  NumNotToSimd);
-			TArrayView<float> AccumulateView(&InAccumulateData[NumToSimd], NumNotToSimd);
+#if INTEL_ISPC
+			ispc::ArraySubtract(InMinuend.GetData(), InSubtrahend.GetData(), OutBuffer.GetData(), Num);
+#endif
+		}
+		else
+		{
+			const int32 NumToSimd = Num & MathIntrinsics::SimdMask;
+			const int32 NumNotToSimd = Num & MathIntrinsics::NotSimdMask;
 
-			ArrayMultiplyAddInPlace(ValuesView, InMultiplier, AccumulateView);
+			if (NumToSimd)
+			{
+				for (int32 i = 0; i < NumToSimd; i += AUDIO_NUM_FLOATS_PER_VECTOR_REGISTER)
+				{
+					VectorRegister4Float Input1 = VectorLoad(&InMinuend[i]);
+					VectorRegister4Float Input2 = VectorLoad(&InSubtrahend[i]);
+					VectorRegister4Float Output = VectorSubtract(Input1, Input2);
+					VectorStore(Output, &OutBuffer[i]);
+				}
+			}
+
+			if (NumNotToSimd)
+			{
+				for (int32 i = NumToSimd; i < Num; ++i)
+				{
+					OutBuffer[i] = InMinuend[i] - InSubtrahend[i];
+				}
+			}
 		}
 	}
 
-	void ArrayLerpAddInPlace(TArrayView<const float> InValues, float InStartMultiplier, float InEndMultiplier, TArrayView<float> InAccumulateValues)
+	/* Subtracts two buffers together element-wise. */
+	void ArraySubtract(const FAlignedFloatBuffer& InMinuend, const FAlignedFloatBuffer& InSubtrahend, FAlignedFloatBuffer& OutputBuffer)
 	{
-		check(InValues.Num() == InAccumulateValues.Num());
+		const int32 InNum = InMinuend.Num();
+		OutputBuffer.Reset(InNum);
+		OutputBuffer.AddUninitialized(InNum);
+
+		TArrayView<const float> InMinuendView(InMinuend.GetData(), InNum);
+		TArrayView<const float> InSubtrahendView(InSubtrahend.GetData(), InNum);
+		TArrayView<float> OutputBufferView(OutputBuffer.GetData(), InNum);
+
+		ArraySubtract(InMinuendView, InSubtrahendView, OutputBufferView);
+	}
+
+	/* Performs element-wise in-place subtraction placing the result in the subtrahend. InOutSubtrahend = InMinuend - InOutSubtrahend */
+	void ArraySubtractInPlace1(TArrayView<const float> InMinuend, TArrayView<float> InOutSubtrahend)
+	{
+		checkf(InMinuend.Num() == InOutSubtrahend.Num(), TEXT("Input buffers must be equal length"));
+
+		const int32 Num = InMinuend.Num();
+
+		if (bAudio_FloatArrayMath_ISPC_Enabled)
+		{
+#if INTEL_ISPC
+			ispc::ArraySubtractInPlace1(InMinuend.GetData(), InOutSubtrahend.GetData(), Num);
+#endif
+		}
+		else
+		{
+			const int32 NumToSimd = Num & MathIntrinsics::SimdMask;
+			const int32 NumNotToSimd = Num & MathIntrinsics::NotSimdMask;
+
+			if (NumToSimd)
+			{
+				for (int32 i = 0; i < NumToSimd; i += AUDIO_NUM_FLOATS_PER_VECTOR_REGISTER)
+				{
+					VectorRegister4Float Input1 = VectorLoad(&InMinuend[i]);
+					VectorRegister4Float Input2 = VectorLoad(&InOutSubtrahend[i]);
+
+					VectorRegister4Float Output = VectorSubtract(Input1, Input2);
+					VectorStore(Output, &InOutSubtrahend[i]);
+				}
+			}
+
+			if (NumNotToSimd)
+			{
+				for (int32 i = NumToSimd; i < Num; ++i)
+				{
+					InOutSubtrahend[i] = InMinuend[i] - InOutSubtrahend[i];
+				}
+			}
+		}
+	}
+
+	/* Performs element-wise in-place subtraction placing the result in the subtrahend. InOutSubtrahend = InMinuend - InOutSubtrahend */
+	void ArraySubtractInPlace1(const FAlignedFloatBuffer& InMinuend, FAlignedFloatBuffer& InOutSubtrahend)
+	{
+		const int32 InNum = InMinuend.Num();
+
+		TArrayView<const float> InMinuendView(InMinuend.GetData(), InNum);
+		TArrayView<float> InOutSubtrahendView(InOutSubtrahend.GetData(), InNum);
+
+		ArraySubtractInPlace1(InMinuendView, InOutSubtrahendView);
+	}
+
+	/* Performs element-wise in-place subtraction placing the result in the minuend. InOutMinuend = InOutMinuend - InSubtrahend */
+	void ArraySubtractInPlace2(TArrayView<float> InOutMinuend, TArrayView<const float> InSubtrahend)
+	{
+		checkf(InOutMinuend.Num() == InSubtrahend.Num(), TEXT("Input buffers must be equal length"));
+
+		const int32 Num = InOutMinuend.Num();
+
+		if (bAudio_FloatArrayMath_ISPC_Enabled)
+		{
+#if INTEL_ISPC
+			ispc::ArraySubtractInPlace2(InOutMinuend.GetData(), InSubtrahend.GetData(), Num);
+#endif
+		}
+		else
+		{
+			const int32 NumToSimd = Num & MathIntrinsics::SimdMask;
+			const int32 NumNotToSimd = Num & MathIntrinsics::NotSimdMask;
+
+			if (NumToSimd)
+			{
+				for (int32 i = 0; i < NumToSimd; i += AUDIO_NUM_FLOATS_PER_VECTOR_REGISTER)
+				{
+					VectorRegister4Float Input1 = VectorLoad(&InOutMinuend[i]);
+					VectorRegister4Float Input2 = VectorLoad(&InSubtrahend[i]);
+
+					VectorRegister4Float Output = VectorSubtract(Input1, Input2);
+					VectorStore(Output, &InOutMinuend[i]);
+				}
+			}
+
+			if (NumNotToSimd)
+			{
+				for (int32 i = NumToSimd; i < Num; ++i)
+				{
+					InOutMinuend[i] = InOutMinuend[i] - InSubtrahend[i];
+				}
+			}
+		}
+	}
+
+	/* Performs element-wise in-place subtraction placing the result in the minuend. InOutMinuend = InOutMinuend - InSubtrahend */
+	void ArraySubtractInPlace2(FAlignedFloatBuffer& InOutMinuend, const FAlignedFloatBuffer& InSubtrahend)
+	{
+		const int32 InNum = InOutMinuend.Num();
+
+		TArrayView<float> InOutMinuendView(InOutMinuend.GetData(), InNum);
+		TArrayView<const float> InSubtrahendView(InSubtrahend.GetData(), InNum);
+
+		ArraySubtractInPlace2(InOutMinuendView, InSubtrahendView);
+	}
+
+	void ArraySubtractByConstantInPlace(TArrayView<float> InValues, float InSubtrahend)
+	{
+		const int32 Num = InValues.Num();
+		float* InData = InValues.GetData();
+		if (bAudio_FloatArrayMath_ISPC_Enabled)
+		{
+#if INTEL_ISPC
+			ispc::ArraySubtractByConstantInPlace(InData, InSubtrahend, Num);
+#endif
+		}
+		else
+		{
+			const int32 NumToSimd = Num & MathIntrinsics::SimdMask;
+			const int32 NumNotToSimd = Num & MathIntrinsics::NotSimdMask;
+
+			const VectorRegister4Float VectorSubtrahend = VectorSetFloat1(InSubtrahend);
+
+			if (NumToSimd)
+			{
+				for (int32 i = 0; i < NumToSimd; i += AUDIO_NUM_FLOATS_PER_VECTOR_REGISTER)
+				{
+					VectorRegister4Float VectorData = VectorLoad(&InData[i]);
+					VectorData = VectorSubtract(VectorData, VectorSubtrahend);
+					VectorStore(VectorData, &InData[i]);
+				}
+			}
+
+			if (NumNotToSimd)
+			{
+				for (int32 i = NumToSimd; i < Num; i++)
+				{
+					InData[i] -= InSubtrahend;
+				}
+			}
+		}
+	}
+
+	void ArraySubtractByConstantInPlace(FAlignedFloatBuffer& InValues, float InSubtrahend)
+	{
+		TArrayView<float> View(InValues.GetData(), InValues.Num());
+
+		ArraySubtractByConstantInPlace(View, InSubtrahend);
+	}
+
+	void ArraySquare(TArrayView<const float> InValues, TArrayView<float> OutValues)
+	{
+		check(InValues.Num() == OutValues.Num());
 
 		const int32 Num = InValues.Num();
 
 		const float* InData = InValues.GetData();
-		float* InAccumulateData = InAccumulateValues.GetData();
+		float* OutData = OutValues.GetData();
+
+		if (bAudio_FloatArrayMath_ISPC_Enabled)
+		{
+#if INTEL_ISPC
+			ispc::ArraySquare(InData, OutData, Num);
+#endif
+		}
+		else
+		{
+			const int32 NumToSimd = Num & MathIntrinsics::SimdMask;
+			const int32 NumNotToSimd = Num & MathIntrinsics::NotSimdMask;
 
-		const float Delta = (InEndMultiplier - InStartMultiplier) / FMath::Max(1.f, static_cast<float>(Num - 1));
-		float Multiplier = InStartMultiplier;
+			if (NumToSimd)
+			{
+				for (int32 i = 0; i < NumToSimd; i += AUDIO_NUM_FLOATS_PER_VECTOR_REGISTER)
+				{
+					VectorRegister4Float VectorData = VectorLoad(&InData[i]);
+					VectorData = VectorMultiply(VectorData, VectorData);
+					VectorStore(VectorData, &OutData[i]);
+				}
+			}
+	
+			if (NumNotToSimd)
+			{
+				for (int32 i = NumToSimd; i < Num; i++)
+				{
+					OutData[i] = InData[i] * InData[i];
+				}
+			}
+		}
+	}
 
-		for (int32 i = 0; i < Num; i++)
+	void ArraySquareInPlace(TArrayView<float> InValues)
+	{
+		const int32 Num = InValues.Num();
+
+		float* InData = InValues.GetData();
+
+		if (bAudio_FloatArrayMath_ISPC_Enabled)
+		{
+#if INTEL_ISPC
+			ispc::ArraySquareInPlace(InData, Num);
+#endif
+		}
+		else
 		{
-			InAccumulateData[i] += InData[i] * Multiplier;
-			Multiplier += Delta;
+			const int32 NumToSimd = Num & MathIntrinsics::SimdMask;
+			const int32 NumNotToSimd = Num & MathIntrinsics::NotSimdMask;
+
+			if (NumToSimd)
+			{
+				for (int32 i = 0; i < NumToSimd; i += AUDIO_NUM_FLOATS_PER_VECTOR_REGISTER)
+				{
+					VectorRegister4Float VectorData = VectorLoad(&InData[i]);
+					VectorData = VectorMultiply(VectorData, VectorData);
+					VectorStore(VectorData, &InData[i]);
+				}
+			}
+	
+			if (NumNotToSimd)
+			{
+				for (int32 i = NumToSimd; i < Num; i++)
+				{
+					InData[i] = InData[i] * InData[i];
+				}
+			}
 		}
 	}
 
-	void ArrayLerpAddInPlace(const FAlignedFloatBuffer& InValues, float InStartMultiplier, float InEndMultiplier, FAlignedFloatBuffer& InAccumulateValues)
+	void ArraySqrtInPlace(TArrayView<float> InValues)
 	{
-		check(InValues.Num() == InAccumulateValues.Num());
+		const int32 Num = InValues.Num();
+		float* InValuesData = InValues.GetData();
+
+		if (bAudio_FloatArrayMath_ISPC_Enabled)
+		{
+#if INTEL_ISPC
+			ispc::ArraySqrtInPlace(InValuesData, Num);
+#endif
+		}
+		else
+		{
+			for (int32 i = 0; i < Num; i++)
+			{
+				InValues[i] = FMath::Sqrt(InValues[i]);
+			}
+		}
+	}
+
+	void ArrayComplexConjugate(TArrayView<const float> InValues, TArrayView<float> OutValues)
+	{
+		check(OutValues.Num() == InValues.Num());
+		check((InValues.Num() % 2) == 0);
 
-		const int32 Num = InAccumulateValues.Num();
-		const int32 NumToSimd = Num & MathIntrinsics::SimdMask;
-		const int32 NumNotToSimd = Num & MathIntrinsics::NotSimdMask;
+		int32 Num = InValues.Num();
 
 		const float* InData = InValues.GetData();
-		float* InAccumulateData = InAccumulateValues.GetData();
+		float* OutData = OutValues.GetData();
+		
+		if (bAudio_FloatArrayMath_ISPC_Enabled)
+		{
+#if INTEL_ISPC
+			ispc::ArrayComplexConjugate(InData, OutData, Num);
+#endif
+		}
+		else
+		{
+			const int32 NumToSimd = Num & MathIntrinsics::SimdMask;
+			const int32 NumNotToSimd = Num & MathIntrinsics::NotSimdMask;
+
+			const VectorRegister4Float ConjugateMult = MakeVectorRegister(1.f, -1.f, 1.f, -1.f);
+
+			if (NumToSimd)
+			{
+				for (int32 i = 0; i < NumToSimd; i += AUDIO_NUM_FLOATS_PER_VECTOR_REGISTER)
+				{
+					VectorRegister4Float VectorData = VectorLoad(&InData[i]);
+
+					VectorData = VectorMultiply(VectorData, ConjugateMult);
+
+					VectorStore(VectorData, &OutData[i]);
+				}
+			}
+
+			if (NumNotToSimd)
+			{
+				for (int32 i = NumToSimd; i < Num; i += 2)
+				{
+					OutData[i] = InData[i];
+					OutData[i + 1] = -InData[i + 1];
+				}
+			}
+		}
+	}
+
+	void ArrayComplexConjugate(const FAlignedFloatBuffer& InValues, FAlignedFloatBuffer& OutValues)
+	{
+		TArrayView<const float> InView(InValues.GetData(), InValues.Num());
+		TArrayView<float> OutView(OutValues.GetData(), OutValues.Num());
+
+		ArrayComplexConjugate(InView, OutView);
+	}
+
+	void ArrayComplexConjugateInPlace(TArrayView<float> InValues)
+	{
+		check((InValues.Num() % 2) == 0);
+
+		int32 Num = InValues.Num();
+
+		float* InData = InValues.GetData();
 		
-		const float Delta = (InEndMultiplier - InStartMultiplier) / FMath::Max(1.f, static_cast<float>(Num - 1));
+		if (bAudio_FloatArrayMath_ISPC_Enabled)
+		{
+#if INTEL_ISPC
+			ispc::ArrayComplexConjugateInPlace(InData, Num);
+#endif
+		}
+		else
+		{
+			const int32 NumToSimd = Num & MathIntrinsics::SimdMask;
+			const int32 NumNotToSimd = Num & MathIntrinsics::NotSimdMask;
+
+			const VectorRegister4Float ConjugateMult = MakeVectorRegister(1.f, -1.f, 1.f, -1.f);
+
+			if (NumToSimd)
+			{
+				for (int32 i = 0; i < NumToSimd; i += AUDIO_NUM_FLOATS_PER_VECTOR_REGISTER)
+				{
+					VectorRegister4Float VectorData = VectorLoad(&InData[i]);
+
+					VectorData = VectorMultiply(VectorData, ConjugateMult);
+
+					VectorStore(VectorData, &InData[i]);
+				}
+			}
+
+			if (NumNotToSimd)
+			{
+				for (int32 i = 1; i < Num; i += 2)
+				{
+					InData[i] *= -1.f;
+				}
+			}
+		}
+	}
+
+	void ArrayComplexConjugateInPlace(FAlignedFloatBuffer& InValues)
+	{
+		TArrayView<float> InView(InValues.GetData(), InValues.Num());
+
+		ArrayComplexConjugateInPlace(InView);
+	}
+
+	void ArrayMagnitudeToDecibelInPlace(TArrayView<float> InValues, float InMinimumDb)
+	{
+		const int32 Num = InValues.Num();
+		float* InValuesData = InValues.GetData();
+
+		if (bAudio_FloatArrayMath_ISPC_Enabled)
+		{
+#if INTEL_ISPC
+			ispc::ArrayMagnitudeToDecibelInPlace(InValuesData, InMinimumDb, Num);
+#endif
+		}
+		else
+		{
+			const int32 NumToSimd = Num & MathIntrinsics::SimdMask;
+			const int32 NumNotToSimd = Num & MathIntrinsics::NotSimdMask;
+
+			const float Scale = 20.f / MathIntrinsics::Loge10;
+			const float Minimum = FMath::Exp(InMinimumDb * MathIntrinsics::Loge10 / 20.f);
+
+			const VectorRegister4Float VectorScale = VectorSetFloat1(Scale);
+			const VectorRegister4Float VectorMinimum = VectorSetFloat1(Minimum);
+
+			if (NumToSimd)
+			{
+				for (int32 i = 0; i < NumToSimd; i += AUDIO_NUM_FLOATS_PER_VECTOR_REGISTER)
+				{
+					VectorRegister4Float VectorData = VectorLoad(&InValuesData[i]);
+
+					VectorData = VectorMax(VectorData, VectorMinimum);
+					VectorData = VectorLog(VectorData);
+					VectorData = VectorMultiply(VectorData, VectorScale);
+
+					VectorStore(VectorData, &InValuesData[i]);
+				}
+			}
+
+			if (NumNotToSimd)
+			{
+				for (int32 i = NumToSimd; i < Num; i++)
+				{
+					InValuesData[i] = FMath::Max(InValuesData[i], Minimum);
+					InValuesData[i] = 20.f * FMath::Loge(InValuesData[i]) / MathIntrinsics::Loge10;
+				}
+			}
+		}
+	}
+
+	void ArrayMagnitudeToDecibelInPlace(FAlignedFloatBuffer& InValues, float InMinimumDb)
+	{
+		TArrayView<float> InView(InValues.GetData(), InValues.Num());
+
+		ArrayMagnitudeToDecibelInPlace(InView, InMinimumDb);
+	}
+
+	void ArrayPowerToDecibelInPlace(TArrayView<float> InValues, float InMinimumDb)
+	{
+		const int32 Num = InValues.Num();
+		float* InValuesData = InValues.GetData();
+
+		if (bAudio_FloatArrayMath_ISPC_Enabled)
+		{
+#if INTEL_ISPC
+			ispc::ArrayPowerToDecibelInPlace(InValuesData, InMinimumDb, Num);
+#endif
+		}
+		else
+		{
+			const int32 NumToSimd = Num & MathIntrinsics::SimdMask;
+			const int32 NumNotToSimd = Num & MathIntrinsics::NotSimdMask;
+
+			const float Scale = 10.f / MathIntrinsics::Loge10;
+			const float Minimum = FMath::Exp(InMinimumDb * MathIntrinsics::Loge10 / 10.f);
+
+			const VectorRegister4Float VectorMinimum = VectorSetFloat1(Minimum);
+			const VectorRegister4Float VectorScale = VectorSetFloat1(Scale);
+
+			if (NumToSimd)
+			{
+				for (int32 i = 0; i < NumToSimd; i += AUDIO_NUM_FLOATS_PER_VECTOR_REGISTER)
+				{
+					VectorRegister4Float VectorData = VectorLoad(&InValuesData[i]);
+
+					VectorData = VectorMax(VectorData, VectorMinimum);
+					VectorData = VectorLog(VectorData);
+					VectorData = VectorMultiply(VectorData, VectorScale);
+
+					VectorStore(VectorData, &InValuesData[i]);
+				}
+			}
+
+			if (NumNotToSimd)
+			{
+				for (int32 i = NumToSimd; i < Num; i++)
+				{
+					InValuesData[i] = FMath::Max(InValuesData[i], Minimum);
+					InValuesData[i] = 10.f * FMath::Loge(InValuesData[i]) / MathIntrinsics::Loge10;
+				}
+			}
+		}
+	}
+
+	void ArrayPowerToDecibelInPlace(FAlignedFloatBuffer& InValues, float InMinimumDb)
+	{
+		TArrayView<float> InView(InValues.GetData(), InValues.Num());
+
+		ArrayPowerToDecibelInPlace(InView, InMinimumDb);
+	}
+
+	void ArrayComplexToPower(TArrayView<const float> InComplexValues, TArrayView<float> OutPowerValues)
+	{
+		check((InComplexValues.Num() % 2) == 0);
+		check(InComplexValues.Num() == (OutPowerValues.Num() * 2));
+
+		const int32 NumOut = OutPowerValues.Num();
 
-		const float FourByDelta = 4.f * Delta;
-		VectorRegister4Float VectorDelta = MakeVectorRegister(FourByDelta, FourByDelta, FourByDelta, FourByDelta);
-		VectorRegister4Float VectorMultiplier = MakeVectorRegister(InStartMultiplier, InStartMultiplier + Delta, InStartMultiplier + 2.f * Delta, InStartMultiplier + 3.f * Delta);
+		const float* InComplexData = InComplexValues.GetData();
+		float* OutPowerData = OutPowerValues.GetData();
 
-		for (int32 i = 0; i < NumToSimd; i += 4)
+		if (bAudio_FloatArrayMath_ISPC_Enabled)
 		{
-			VectorRegister4Float VectorData = VectorLoadAligned(&InData[i]);
-			VectorRegister4Float VectorAccumData = VectorLoadAligned(&InAccumulateData[i]);
+#if INTEL_ISPC
+			ispc::ArrayComplexToPowerInterleaved(InComplexData, OutPowerData, NumOut);
+#endif
+		}
+		else
+		{
+			const int32 NumToSimd = NumOut & MathIntrinsics::SimdMask;
+			const int32 NumNotToSimd = NumOut & MathIntrinsics::NotSimdMask;
+
+			if (NumToSimd)
+			{
+				for (int32 i = 0; i < NumToSimd; i += AUDIO_NUM_FLOATS_PER_VECTOR_REGISTER)
+				{
+					VectorRegister4Float VectorComplex1 = VectorLoad(&InComplexData[2 * i]);
+					VectorRegister4Float VectorSquared1 = VectorMultiply(VectorComplex1, VectorComplex1);
+
+					VectorRegister4Float VectorComplex2 = VectorLoad(&InComplexData[(2 * i) + 4]);
+					VectorRegister4Float VectorSquared2 = VectorMultiply(VectorComplex2, VectorComplex2);
+
+					VectorRegister4Float VectorSquareReal = VectorShuffle(VectorSquared1, VectorSquared2, 0, 2, 0, 2);
+					VectorRegister4Float VectorSquareImag = VectorShuffle(VectorSquared1, VectorSquared2, 1, 3, 1, 3);
 
-			VectorRegister4Float VectorOut = VectorMultiplyAdd(VectorData, VectorMultiplier, VectorAccumData);
-			VectorMultiplier = VectorAdd(VectorMultiplier, VectorDelta);
+					VectorRegister4Float VectorOut = VectorAdd(VectorSquareReal, VectorSquareImag);
+
+					VectorStore(VectorOut, &OutPowerData[i]);
+				}
+			}
+
+			if (NumNotToSimd)
+			{
+				for (int32 i = NumToSimd; i < NumOut; i++)
+				{
+					int32 ComplexPos = 2 * i;
+
+					float RealValue = InComplexData[ComplexPos];
+					float ImagValue = InComplexData[ComplexPos + 1];
+
+					OutPowerData[i] = (RealValue * RealValue) + (ImagValue * ImagValue);
+				}
+			}
+		}
+	}
+
+	void ArrayComplexToPower(const FAlignedFloatBuffer& InComplexValues, FAlignedFloatBuffer& OutPowerValues)
+	{
+		TArrayView<const float> ComplexView(InComplexValues.GetData(), InComplexValues.Num());
+		TArrayView<float> PowerView(OutPowerValues.GetData(), OutPowerValues.Num());
+
+		ArrayComplexToPower(ComplexView, PowerView);
+	}
+
+	void ArrayComplexToPower(TArrayView<const float> InRealSamples, TArrayView<const float> InImaginarySamples, TArrayView<float> OutPowerSamples)
+	{
+		checkf(InRealSamples.Num() == InImaginarySamples.Num(), TEXT("Input buffers must have equal number of elements"));
+
+		const int32 Num = InRealSamples.Num();
+
+		if (bAudio_FloatArrayMath_ISPC_Enabled)
+		{
+#if INTEL_ISPC
+			ispc::ArrayComplexToPower(InRealSamples.GetData(), InImaginarySamples.GetData(), OutPowerSamples.GetData(), Num);
+#endif
+		}
+		else
+		{
+			const int32 NumToSimd = Num & MathIntrinsics::SimdMask;
+			const int32 NumNotToSimd = Num & MathIntrinsics::NotSimdMask;
+
+			if (NumToSimd)
+			{
+				for (int32 i = 0; i < NumToSimd; i += AUDIO_NUM_FLOATS_PER_VECTOR_REGISTER)
+				{
+					VectorRegister4Float VInReal = VectorLoad(&InRealSamples[i]);
+					VectorRegister4Float VInRealSquared = VectorMultiply(VInReal, VInReal);
+
+					VectorRegister4Float VInImag = VectorLoad(&InImaginarySamples[i]);
+					VectorRegister4Float VOut = VectorMultiplyAdd(VInImag, VInImag, VInRealSquared);
+
+					VectorStore(VOut, &OutPowerSamples[i]);
+				}
+			}
+
+			if (NumNotToSimd)
+			{
+				for (int32 i = NumToSimd; i < Num; ++i)
+				{
+					const float InRealSquared = InRealSamples[i] * InRealSamples[i];
+					const float InImagSquared = InImaginarySamples[i] * InImaginarySamples[i];
+
+					OutPowerSamples[i] = InRealSquared + InImagSquared;
+				}
+			}
+		}
+	}
+
+	void ArrayComplexToPower(const FAlignedFloatBuffer& InRealSamples, const FAlignedFloatBuffer& InImaginarySamples, FAlignedFloatBuffer& OutPowerSamples)
+	{
+		const int32 Num = InRealSamples.Num();
 
-			VectorStoreAligned(VectorOut, &InAccumulateData[i]);
-		}
+		OutPowerSamples.Reset(Num);
+		OutPowerSamples.AddUninitialized(Num);
 
-		if (NumNotToSimd)
-		{
-			TArrayView<const float> ValuesView(&InData[NumToSimd],  NumNotToSimd);
-			TArrayView<float> AccumulateView(&InAccumulateData[NumToSimd], NumNotToSimd);
+		TArrayView<const float> InRealSamplesView(InRealSamples.GetData(), Num);
+		TArrayView<const float> InImaginarySamplesView(InImaginarySamples.GetData(), Num);
+		TArrayView<float> OutPowerSamplesView(OutPowerSamples.GetData(), Num);
 
-			ArrayLerpAddInPlace(ValuesView, InStartMultiplier + NumToSimd * Delta, InEndMultiplier, AccumulateView);
-		}
+		ArrayComplexToPower(InRealSamplesView, InImaginarySamplesView, OutPowerSamplesView);
 	}
 
-	void ArraySubtractByConstantInPlace(TArrayView<float> InValues, float InSubtrahend)
+	/* Sets a values to zero if value is denormal. Denormal numbers significantly slow down floating point operations. */
+	void ArrayUnderflowClamp(TArrayView<float> InOutValues)
 	{
-		const int32 Num = InValues.Num();
-		float* InValuesData = InValues.GetData();
-		for (int32 i = 0; i < Num; i++)
+		int32 Num = InOutValues.Num();
+		float* InOutData = InOutValues.GetData();
+
+		if (bAudio_FloatArrayMath_ISPC_Enabled)
 		{
-			InValuesData[i] -= InSubtrahend;
+#if INTEL_ISPC
+			ispc::ArrayUnderflowClamp(InOutData, Num);
+#endif
 		}
-	}
+		else
+		{
+			const int32 NumToSimd = Num & MathIntrinsics::SimdMask;
+			const int32 NumNotToSimd = Num & MathIntrinsics::NotSimdMask;
 
-	void ArraySubtractByConstantInPlace(FAlignedFloatBuffer& InValues, float InSubtrahend)
-	{
-		const int32 Num = InValues.Num();
-		const int32 NumToSimd = Num & MathIntrinsics::SimdMask;
-		const int32 NumNotToSimd = Num & MathIntrinsics::NotSimdMask;
+			const VectorRegister4Float VFMIN = MakeVectorRegister(FLT_MIN, FLT_MIN, FLT_MIN, FLT_MIN);
+			const VectorRegister4Float VNFMIN = MakeVectorRegister(-FLT_MIN, -FLT_MIN, -FLT_MIN, -FLT_MIN);
 
-		float* InData = InValues.GetData();
+			if (NumToSimd)
+			{
+				for (int32 i = 0; i < NumToSimd; i += AUDIO_NUM_FLOATS_PER_VECTOR_REGISTER)
+				{
+					VectorRegister4Float VInOut = VectorLoad(&InOutData[i]);
 
-		const VectorRegister4Float VectorSubtrahend = VectorSetFloat1(InSubtrahend);
+					// Create mask of denormal numbers.
+					VectorRegister4Float Mask = VectorBitwiseAnd(VectorCompareGT(VInOut, VNFMIN), VectorCompareLT(VInOut, VFMIN));
 
-		for (int32 i = 0; i < NumToSimd; i += 4)
-		{
-			VectorRegister4Float VectorData = VectorLoadAligned(&InData[i]);
-			VectorData = VectorSubtract(VectorData, VectorSubtrahend);
-			VectorStoreAligned(VectorData, &InData[i]);
-		}
+					// Choose between zero or original number based upon mask.
+					VInOut = VectorSelect(Mask, GlobalVectorConstants::FloatZero, VInOut);
+					VectorStore(VInOut, &InOutData[i]);
+				}
+			}
 
-		if (NumNotToSimd)
-		{
-			TArrayView<float> View(&InData[NumToSimd], NumNotToSimd);
-			ArraySubtractByConstantInPlace(View, InSubtrahend);
+			if (NumNotToSimd)
+			{
+				for (int32 i = NumToSimd; i < Num; i++)
+				{
+					float InOut = InOutData[i];
+
+					// Create mask of denormal numbers.
+					const bool Mask = (InOut > -FLT_MIN) && (InOut < FLT_MIN);
+
+					// Choose between zero or original number based upon mask.
+					InOut = Mask ? 0.0f : InOut;
+					InOutData[i] = InOut;
+				}
+			}
 		}
 	}
 
-	void ArraySubtract(TArrayView<const float> InMinuend, TArrayView<const float> InSubtrahend, TArray<float>& OutArray)
+	void ArrayUnderflowClamp(FAlignedFloatBuffer& InOutBuffer)
 	{
-		const int32 Num = InMinuend.Num();
+		TArrayView<float> InOutBufferView(InOutBuffer.GetData(), InOutBuffer.Num());
 
-		checkf(Num == InSubtrahend.Num() && Num == OutArray.Num(), TEXT("InMinuend, InSubtrahend, and OutArray must have equal Num elements (%d vs %d vs %d)"), Num, InSubtrahend.Num(), OutArray.Num());
-		
+		ArrayUnderflowClamp(InOutBufferView);
+	}
 
-		if (Num < 1)
-		{
-			return;
-		}
-		
-		const float* MinuendPtr = InMinuend.GetData();
-		const float* SubtrahendPtr = InSubtrahend.GetData();
-		float* OutPtr = OutArray.GetData();
+	/* Clamps values in the buffer to be between InMinValue and InMaxValue */
+	void ArrayRangeClamp(TArrayView<float> InOutBuffer, float InMinValue, float InMaxValue)
+	{
+		int32 Num = InOutBuffer.Num();
+		float* InOutData = InOutBuffer.GetData();
 
-		for (int32 i = 0; i < Num; i++)
+		if (bAudio_FloatArrayMath_ISPC_Enabled)
 		{
-			OutPtr[i] = MinuendPtr[i] - SubtrahendPtr[i];
+#if INTEL_ISPC
+			ispc::ArrayRangeClamp(InOutData, Num, InMinValue, InMaxValue);
+#endif
 		}
-	}
+		else
+		{
+			const int32 NumToSimd = Num & MathIntrinsics::SimdMask;
+			const int32 NumNotToSimd = Num & MathIntrinsics::NotSimdMask;
 
-	void ArraySubtract(const FAlignedFloatBuffer& InMinuend, const FAlignedFloatBuffer& InSubtrahend, FAlignedFloatBuffer& OutArray)
-	{
-		BufferSubtractFast(InMinuend, InSubtrahend, OutArray);
-	}
+			const VectorRegister4Float VMinVal = MakeVectorRegister(InMinValue, InMinValue, InMinValue, InMinValue);
+			const VectorRegister4Float VMaxVal = MakeVectorRegister(InMaxValue, InMaxValue, InMaxValue, InMaxValue);
 
-	void ArraySquare(TArrayView<const float> InValues, TArrayView<float> OutValues)
-	{
-		check(InValues.Num() == OutValues.Num());
+			if (NumToSimd)
+			{
+				for (int32 i = 0; i < NumToSimd; i += AUDIO_NUM_FLOATS_PER_VECTOR_REGISTER)
+				{
+					VectorRegister4Float VInOut = VectorLoad(&InOutData[i]);
 
-		const uint32 Num = InValues.Num();
-		const uint32 NumToSimd = Num & MathIntrinsics::SimdMask;
+					// Create masks to flag elements outside of range.
+					VectorRegister4Float MinMask = VectorCompareLT(VInOut, VMinVal);
+					VectorRegister4Float MaxMask = VectorCompareGT(VInOut, VMaxVal);
 
-		const float* InData = InValues.GetData();
-		float* OutData = OutValues.GetData();
+					// Choose between range extremes or original number based on masks.
+					VInOut = VectorSelect(MinMask, VMinVal, VInOut);
+					VInOut = VectorSelect(MaxMask, VMaxVal, VInOut);
 
-		for (uint32 i = 0; i < NumToSimd; i += 4)
-		{
-			VectorRegister4Float VectorData = VectorLoad(&InData[i]);
-			VectorData = VectorMultiply(VectorData, VectorData);
-			VectorStore(VectorData, &OutData[i]);
-		}
+					VectorStore(VInOut, &InOutData[i]);
+				}
+			}
 
-		for (uint32 i = NumToSimd; i < Num; i++)
-		{
-			OutData[i] = InData[i] * InData[i];
+			if (NumNotToSimd)
+			{
+				for (int32 i = NumToSimd; i < Num; i++)
+				{
+					InOutData[i] = FMath::Clamp(InOutData[i], InMinValue, InMaxValue);
+				}
+			}
 		}
 	}
 
-	void ArraySquareInPlace(TArrayView<float> InValues)
+	void ArrayRangeClamp(FAlignedFloatBuffer& InOutBuffer, float InMinValue, float InMaxValue)
 	{
-		const uint32 Num = InValues.Num();
-		const uint32 NumToSimd = Num & MathIntrinsics::SimdMask;
-
-		float* InData = InValues.GetData();
+		TArrayView<float> InOutBufferView(InOutBuffer.GetData(), InOutBuffer.Num());
 
-		for (uint32 i = 0; i < NumToSimd; i += 4)
-		{
-			VectorRegister4Float VectorData = VectorLoad(&InData[i]);
-			VectorData = VectorMultiply(VectorData, VectorData);
-			VectorStore(VectorData, &InData[i]);
-		}
-
-		for (uint32 i = NumToSimd; i < Num; i++)
-		{
-			InData[i] = InData[i] * InData[i];
-		}
+		ArrayRangeClamp(InOutBufferView, InMinValue, InMaxValue);
 	}
 
-	void ArraySqrtInPlace(TArrayView<float> InValues)
+	void ArraySetToConstantInplace(TArrayView<float> InOutBuffer, float InConstant)
 	{
-		const int32 Num = InValues.Num();
-		float* InValuesData = InValues.GetData();
+		int32 Num = InOutBuffer.Num();
+		float* InOutData = InOutBuffer.GetData();
 
-		for (int32 i = 0; i < Num; i++)
+		if (bAudio_FloatArrayMath_ISPC_Enabled)
 		{
-			InValues[i] = FMath::Sqrt(InValues[i]);
+#if INTEL_ISPC
+			ispc::ArraySetToConstantInplace(InOutData, Num, InConstant);
+#endif
 		}
-	}
+		else
+		{
+			const int32 NumToSimd = Num & MathIntrinsics::SimdMask;
+			const int32 NumNotToSimd = Num & MathIntrinsics::NotSimdMask;
 
-	void ArrayComplexConjugate(TArrayView<const float> InValues, TArrayView<float> OutValues)
-	{
-		check(OutValues.Num() == InValues.Num());
-		check((InValues.Num() % 2) == 0);
+			const VectorRegister4Float Constant = VectorLoadFloat1(&InConstant);
 
-		int32 Num = InValues.Num();
+			if (NumToSimd)
+			{
+				for (int32 i = 0; i < NumToSimd; i += AUDIO_NUM_FLOATS_PER_VECTOR_REGISTER)
+				{
+					VectorStore(Constant, &InOutData[i]);
+				}
+			}
 
-		const float* InData = InValues.GetData();
-		float* OutData = OutValues.GetData();
-		
-		for (int32 i = 0; i < Num; i+= 2)
-		{
-			OutData[i] = InData[i];
-			OutData[i + 1] = -InData[i + 1];
+			if (NumNotToSimd)
+			{
+				for (int32 i = NumToSimd; i < Num; i++)
+				{
+					InOutData[i] = InConstant;
+				}
+			}
 		}
 	}
 
-	void ArrayComplexConjugate(const FAlignedFloatBuffer& InValues, FAlignedFloatBuffer& OutValues)
+	void ArraySetToConstantInplace(FAlignedFloatBuffer& InOutBuffer, float InConstant)
 	{
-		check(OutValues.Num() == InValues.Num());
+		TArrayView<float> InOutBufferView(InOutBuffer.GetData(), InOutBuffer.Num());
 
-		const int32 Num = InValues.Num();
-		const int32 NumToSimd = Num & MathIntrinsics::SimdMask;
-		const int32 NumNotToSimd = Num & MathIntrinsics::NotSimdMask;
+		ArraySetToConstantInplace(InOutBufferView, InConstant);
+	}
 
-		const float* InData = InValues.GetData();
-		float* OutData = OutValues.GetData();
+	/* Performs an element-wise weighted sum OutputBuffer = (InBuffer1 x InGain1) + (InBuffer2 x InGain2) */
+	void ArrayWeightedSum(TArrayView<const float> InBuffer1, float InGain1, TArrayView<const float> InBuffer2, float InGain2, TArrayView<float> OutBuffer)
+	{
+		checkf(InBuffer1.Num() == InBuffer2.Num(), TEXT("Buffers must be equal length"));
 
-		const VectorRegister4Float ConjugateMult = MakeVectorRegister(1.f, -1.f, 1.f, -1.f);
+		int32 Num = InBuffer1.Num();
+		const float* InData1 = InBuffer1.GetData();
+		const float* InData2 = InBuffer2.GetData();
+		float* OutData = OutBuffer.GetData();
 
-		for (int32 i = 0; i < NumToSimd; i += 4)
+		if (bAudio_FloatArrayMath_ISPC_Enabled)
 		{
-			VectorRegister4Float VectorData = VectorLoadAligned(&InData[i]);
-			
-			VectorData = VectorMultiply(VectorData, ConjugateMult);
-
-			VectorStoreAligned(VectorData, &OutData[i]);
+#if INTEL_ISPC
+			ispc::ArrayWeightedSumTwoGain(InData1, InGain1, InData2, InGain2, OutData, Num);
+#endif
 		}
-
-		if (NumNotToSimd)
+		else
 		{
-			TArrayView<const float> InView(&InData[NumToSimd], NumNotToSimd);
-			TArrayView<float> OutView(&OutData[NumToSimd], NumNotToSimd);
+			const int32 NumToSimd = Num & MathIntrinsics::SimdMask;
+			const int32 NumNotToSimd = Num & MathIntrinsics::NotSimdMask;
+
+			if (NumToSimd)
+			{
+				VectorRegister4Float Gain1Vector = VectorLoadFloat1(&InGain1);
+				VectorRegister4Float Gain2Vector = VectorLoadFloat1(&InGain2);
+
+				for (int32 i = 0; i < NumToSimd; i += AUDIO_NUM_FLOATS_PER_VECTOR_REGISTER)
+				{
+					// InBuffer1 x InGain1
+					VectorRegister4Float Input1 = VectorLoad(&InData1[i]);
+
+					// InBuffer2 x InGain2
+					VectorRegister4Float Input2 = VectorLoad(&InData2[i]);
+					VectorRegister4Float Weighted2 = VectorMultiply(Input2, Gain2Vector);
 
-			ArrayComplexConjugate(InView, OutView);
+					VectorRegister4Float Output = VectorMultiplyAdd(Input1, Gain1Vector, Weighted2);
+					VectorStore(Output, &OutData[i]);
+				}
+			}
+
+			if (NumNotToSimd)
+			{
+				for (int32 i = NumToSimd; i < Num; i++)
+				{
+					OutData[i] = (InData1[i] * InGain1) + (InData2[i] * InGain2);
+				}
+			}
 		}
 	}
 
-	void ArrayComplexConjugateInPlace(TArrayView<float> InValues)
+	/* Performs an element-wise weighted sum OutputBuffer = (InBuffer1 x InGain1) + (InBuffer2 x InGain2) */
+	void ArrayWeightedSum(const FAlignedFloatBuffer& InBuffer1, float InGain1, const FAlignedFloatBuffer& InBuffer2, float InGain2, FAlignedFloatBuffer& OutBuffer)
 	{
-		check((InValues.Num() % 2) == 0);
+		OutBuffer.Reset();
+		OutBuffer.AddUninitialized(InBuffer1.Num());
 
-		int32 Num = InValues.Num();
+		TArrayView<const float> InBuffer1View(InBuffer1.GetData(), InBuffer1.Num());
+		TArrayView<const float> InBuffer2View(InBuffer2.GetData(), InBuffer2.Num());
+		TArrayView<float> OutBufferView(OutBuffer.GetData(), OutBuffer.Num());
 
-		float* InData = InValues.GetData();
-		
-		for (int32 i = 1; i < Num; i+= 2)
-		{
-			InData[i] *= -1.f;
-		}
+		ArrayWeightedSum(InBuffer1View, InGain1, InBuffer2View, InGain2, OutBufferView);
 	}
 
-	void ArrayComplexConjugateInPlace(FAlignedFloatBuffer& InValues)
+	/* Performs an element-wise weighted sum OutputBuffer = (InBuffer1 x InGain1) + InBuffer2 */
+	void ArrayWeightedSum(TArrayView<const float> InBuffer1, float InGain1, TArrayView<const float> InBuffer2, TArrayView<float> OutBuffer)
 	{
-		const int32 Num = InValues.Num();
-		const int32 NumToSimd = Num & MathIntrinsics::SimdMask;
-		const int32 NumNotToSimd = Num & MathIntrinsics::NotSimdMask;
+		checkf(InBuffer1.Num() == InBuffer2.Num() && InBuffer1.Num() == OutBuffer.Num(), TEXT("Buffers must be equal length"));
 
-		float* InData = InValues.GetData();
-
-		const VectorRegister4Float ConjugateMult = MakeVectorRegister(1.f, -1.f, 1.f, -1.f);
+		int32 Num = InBuffer1.Num();
+		const float* InData1 = InBuffer1.GetData();
+		const float* InData2 = InBuffer2.GetData();
+		float* OutData = OutBuffer.GetData();
 
-		for (int32 i = 0; i < NumToSimd; i += 4)
+		if (bAudio_FloatArrayMath_ISPC_Enabled)
 		{
-			VectorRegister4Float VectorData = VectorLoadAligned(&InData[i]);
-			
-			VectorData = VectorMultiply(VectorData, ConjugateMult);
-
-			VectorStoreAligned(VectorData, &InData[i]);
+#if INTEL_ISPC
+			ispc::ArrayWeightedSumOneGain(InData1, InGain1, InData2, OutData, Num);
+#endif
 		}
-
-		if (NumNotToSimd)
+		else
 		{
-			TArrayView<float> InView(&InData[NumToSimd], NumNotToSimd);
+			const int32 NumToSimd = Num & MathIntrinsics::SimdMask;
+			const int32 NumNotToSimd = Num & MathIntrinsics::NotSimdMask;
+
+			if (NumToSimd)
+			{
+				VectorRegister4Float Gain1Vector = VectorLoadFloat1(&InGain1);
 
-			ArrayComplexConjugateInPlace(InView);
+				for (int32 i = 0; i < NumToSimd; i += AUDIO_NUM_FLOATS_PER_VECTOR_REGISTER)
+				{
+					// InBuffer1 x InGain1
+					VectorRegister4Float Input1 = VectorLoad(&InData1[i]);
+					VectorRegister4Float Input2 = VectorLoad(&InData2[i]);
+					VectorRegister4Float Output = VectorMultiplyAdd(Input1, Gain1Vector, Input2);
+					VectorStore(Output, &OutData[i]);
+				}
+			}
+
+			if (NumNotToSimd)
+			{
+				for (int32 i = NumToSimd; i < Num; i++)
+				{
+					OutData[i] = (InData1[i] * InGain1) + InData2[i];
+				}
+			}
 		}
 	}
 
-	void ArrayMagnitudeToDecibelInPlace(TArrayView<float> InValues, float InMinimumDb)
+	/* Performs an element-wise weighted sum OutputBuffer = (InBuffer1 x InGain1) + InBuffer2 */
+	void ArrayWeightedSum(const FAlignedFloatBuffer& InBuffer1, float InGain1, const FAlignedFloatBuffer& InBuffer2, FAlignedFloatBuffer& OutBuffer)
 	{
-		const int32 Num = InValues.Num();
-		float* InValuesData = InValues.GetData();
+		OutBuffer.Reset();
+		OutBuffer.AddUninitialized(InBuffer1.Num());
 
-		const float Minimum = FMath::Exp(InMinimumDb * MathIntrinsics::Loge10 / 20.f);
+		TArrayView<const float> InBuffer1View(InBuffer1.GetData(), InBuffer1.Num());
+		TArrayView<const float> InBuffer2View(InBuffer2.GetData(), InBuffer2.Num());
+		TArrayView<float> OutBufferView(OutBuffer.GetData(), OutBuffer.Num());
 
-		for (int32 i = 0; i < Num; i++)
-		{
-			InValuesData[i] = FMath::Max(InValuesData[i], Minimum);
-			InValuesData[i] = 20.f * FMath::Loge(InValuesData[i]) / MathIntrinsics::Loge10;
-		}
+		ArrayWeightedSum(InBuffer1View, InGain1, InBuffer2View, OutBufferView);
 	}
 
-	void ArrayMagnitudeToDecibelInPlace(FAlignedFloatBuffer& InValues, float InMinimumDb)
+	void ArrayFade(TArrayView<float> InOutBuffer, const float StartValue, const float EndValue)
 	{
-		const int32 Num = InValues.Num();
-		const int32 NumToSimd = Num & MathIntrinsics::SimdMask;
-		const int32 NumNotToSimd = Num & MathIntrinsics::NotSimdMask;
+		int32 Num = InOutBuffer.Num();
+		float* OutFloatBuffer = InOutBuffer.GetData();
 
-		float* InData = InValues.GetData();
+		if (FMath::IsNearlyEqual(StartValue, EndValue))
+		{
+			// No need to do anything if start and end values are both 0.0
+			if (StartValue == 0.0f)
+			{
+				FMemory::Memset(OutFloatBuffer, 0, sizeof(float) * Num);
+			}
+			else
+			{
+				ArrayMultiplyByConstantInPlace(InOutBuffer, StartValue);
+			}
+		}
+		else
+		{
+			if (bAudio_FloatArrayMath_ISPC_Enabled)
+			{
+#if INTEL_ISPC
+				ispc::ArrayFade(OutFloatBuffer, Num, StartValue, EndValue);
+#endif
+			}
+			else
+			{
+				const int32 NumToSimd = Num & MathIntrinsics::SimdMask;
+				const int32 NumNotToSimd = Num & MathIntrinsics::NotSimdMask;
 
-		const float Scale = 20.f / MathIntrinsics::Loge10;
-		const float Minimum = FMath::Exp(InMinimumDb * MathIntrinsics::Loge10 / 20.f);
+				if (NumToSimd)
+				{
+					const int32 NumIterations = Num / AUDIO_NUM_FLOATS_PER_VECTOR_REGISTER;
+					const float DeltaValue = ((EndValue - StartValue) / NumIterations);
 
-		const VectorRegister4Float VectorScale = VectorSetFloat1(Scale);
-		const VectorRegister4Float VectorMinimum = VectorSetFloat1(Minimum);
+					VectorRegister4Float Gain = VectorLoadFloat1(&StartValue);
+					VectorRegister4Float Delta = VectorLoadFloat1(&DeltaValue);
 
-		for (int32 i = 0; i < NumToSimd; i += 4)
-		{
-			VectorRegister4Float VectorData = VectorLoadAligned(&InData[i]);
-			
-			VectorData = VectorMax(VectorData, VectorMinimum);
-			VectorData = VectorLog(VectorData);
-			VectorData = VectorMultiply(VectorData, VectorScale);
+					for (int32 i = 0; i < NumToSimd; i += AUDIO_NUM_FLOATS_PER_VECTOR_REGISTER)
+					{
+						VectorRegister4Float Output = VectorLoad(&OutFloatBuffer[i]);
+						Output = VectorMultiply(Output, Gain);
+						Gain = VectorAdd(Gain, Delta);
+						VectorStore(Output, &OutFloatBuffer[i]);
+					}
+				}
 
-			VectorStoreAligned(VectorData, &InData[i]);
-		}
+				if (NumNotToSimd)
+				{
+					float Gain = StartValue;
 
-		if (NumNotToSimd)
-		{
-			TArrayView<float> InView(&InData[NumToSimd], NumNotToSimd);
-			ArrayMagnitudeToDecibelInPlace(InView, InMinimumDb);
+					// Do a fade from start to end
+					const float DeltaValue = ((EndValue - StartValue) / Num);
+					for (int32 i = NumToSimd; i < Num; ++i)
+					{
+						OutFloatBuffer[i] = OutFloatBuffer[i] * Gain;
+						Gain += DeltaValue;
+					}
+				}
+			}
 		}
 	}
 
-	void ArrayPowerToDecibelInPlace(TArrayView<float> InValues, float InMinimumDb)
+	void ArrayFade(FAlignedFloatBuffer& OutFloatBuffer, const float StartValue, const float EndValue)
 	{
-		const int32 Num = InValues.Num();
-		float* InValuesData = InValues.GetData();
+		TArrayView<float> OutFloatBufferView(OutFloatBuffer.GetData(), OutFloatBuffer.Num());
 
-		const float Minimum = FMath::Exp(InMinimumDb * MathIntrinsics::Loge10 / 10.f);
-
-		for (int32 i = 0; i < Num; i++)
-		{
-			InValuesData[i] = FMath::Max(InValuesData[i], Minimum);
-			InValuesData[i] = 10.f * FMath::Loge(InValuesData[i]) / MathIntrinsics::Loge10;
-		}
+		ArrayFade(OutFloatBufferView, StartValue, EndValue);
 	}
 
-	void ArrayPowerToDecibelInPlace(FAlignedFloatBuffer& InValues, float InMinimumDb)
+	void ArrayMixIn(TArrayView<const float> InFloatBuffer, TArrayView<float> BufferToSumTo, const float Gain)
 	{
-		const int32 Num = InValues.Num();
-		const int32 NumToSimd = Num & MathIntrinsics::SimdMask;
-		const int32 NumNotToSimd = Num & MathIntrinsics::NotSimdMask;
-
-		float* InData = InValues.GetData();
-
-		const float Scale = 10.f / MathIntrinsics::Loge10;
-		const float Minimum = FMath::Exp(InMinimumDb * MathIntrinsics::Loge10 / 10.f);
+		checkf(InFloatBuffer.Num() == BufferToSumTo.Num(), TEXT("Buffers must be equal size"));
 
-		const VectorRegister4Float VectorMinimum = VectorSetFloat1(Minimum);
-		const VectorRegister4Float VectorScale = VectorSetFloat1(Scale);
+		int32 Num = InFloatBuffer.Num();
+		const float* InData = InFloatBuffer.GetData();
+		float* InOutData = BufferToSumTo.GetData();
 
-		for (int32 i = 0; i < NumToSimd; i += 4)
+		if (bAudio_FloatArrayMath_ISPC_Enabled)
 		{
-			VectorRegister4Float VectorData = VectorLoadAligned(&InData[i]);
+#if INTEL_ISPC
+			ispc::ArrayMixInWithGain(InData, InOutData, Num, Gain);
+#endif
+		}
+		else
+		{
+			const int32 NumToSimd = Num & MathIntrinsics::SimdMask;
+			const int32 NumNotToSimd = Num & MathIntrinsics::NotSimdMask;
 
-			VectorData = VectorMax(VectorData, VectorMinimum);
-			VectorData = VectorLog(VectorData);
-			VectorData = VectorMultiply(VectorData, VectorScale);
+			if (NumToSimd)
+			{
+				VectorRegister4Float GainVector = VectorLoadFloat1(&Gain);
 
-			VectorStoreAligned(VectorData, &InData[i]);
-		}
+				for (int32 i = 0; i < NumToSimd; i += AUDIO_NUM_FLOATS_PER_VECTOR_REGISTER)
+				{
+					VectorRegister4Float Output = VectorLoad(&InOutData[i]);
+					VectorRegister4Float Input = VectorLoad(&InData[i]);
+					Output = VectorMultiplyAdd(Input, GainVector, Output);
+					VectorStore(Output, &InOutData[i]);
+				}
+			}
 
-		if (NumNotToSimd)
-		{
-			TArrayView<float> InView(&InData[NumToSimd], NumNotToSimd);
-			ArrayPowerToDecibelInPlace(InView, InMinimumDb);
+			if (NumNotToSimd)
+			{
+				for (int32 i = NumToSimd; i < Num; ++i)
+				{
+					InOutData[i] += InData[i] * Gain;
+				}
+			}
 		}
 	}
 
-	void ArrayComplexToPower(TArrayView<const float> InComplexValues, TArrayView<float> OutPowerValues)
+	void ArrayMixIn(const FAlignedFloatBuffer& InFloatBuffer, FAlignedFloatBuffer& BufferToSumTo, const float Gain)
 	{
-		check((InComplexValues.Num() % 2) == 0);
-		check(InComplexValues.Num() == (OutPowerValues.Num() * 2));
+		TArrayView<const float> InFloatBufferView(InFloatBuffer.GetData(), InFloatBuffer.Num());
+		TArrayView<float> BufferToSumToView(BufferToSumTo.GetData(), BufferToSumTo.Num());
 
-		const int32 NumOut = OutPowerValues.Num();
+		ArrayMixIn(InFloatBufferView, BufferToSumToView, Gain);
+	}
 
-		const float* InComplexData = InComplexValues.GetData();
-		float* OutPowerData = OutPowerValues.GetData();
+	void ArrayMixIn(TArrayView<const float> InFloatBuffer, TArrayView<float> BufferToSumTo)
+	{
+		checkf(InFloatBuffer.Num() == BufferToSumTo.Num(), TEXT("Buffers must be equal size"));
+
+		int32 Num = InFloatBuffer.Num();
+		const float* InData = InFloatBuffer.GetData();
+		float* InOutData = BufferToSumTo.GetData();
 
-		for (int32 i = 0; i < NumOut; i++)
+		if (bAudio_FloatArrayMath_ISPC_Enabled)
+		{
+#if INTEL_ISPC
+			ispc::ArrayMixIn(InData, InOutData, Num);
+#endif
+		}
+		else
 		{
-			int32 ComplexPos = 2 * i;
+			const int32 NumToSimd = Num & MathIntrinsics::SimdMask;
+			const int32 NumNotToSimd = Num & MathIntrinsics::NotSimdMask;
 
-			float RealValue = InComplexData[ComplexPos];
-			float ImagValue = InComplexData[ComplexPos + 1];
+			if (NumToSimd)
+			{
+				for (int32 i = 0; i < NumToSimd; i += AUDIO_NUM_FLOATS_PER_VECTOR_REGISTER)
+				{
+					VectorRegister4Float Output = VectorLoad(&InOutData[i]);
+					VectorRegister4Float Input = VectorLoad(&InData[i]);
+					Output = VectorAdd(Input, Output);
+					VectorStore(Output, &InOutData[i]);
+				}
+			}
 
-			OutPowerData[i] = (RealValue * RealValue) + (ImagValue * ImagValue);
+			if (NumNotToSimd)
+			{
+				for (int32 i = NumToSimd; i < Num; ++i)
+				{
+					InOutData[i] += InData[i];
+				}
+			}
 		}
 	}
 
-	void ArrayComplexToPower(const FAlignedFloatBuffer& InComplexValues, FAlignedFloatBuffer& OutPowerValues)
+	void ArrayMixIn(const FAlignedFloatBuffer& InFloatBuffer, FAlignedFloatBuffer& BufferToSumTo)
 	{
-		check((InComplexValues.Num() % 2) == 0);
-		check(InComplexValues.Num() == (OutPowerValues.Num() * 2));
+		TArrayView<const float> InFloatBufferView(InFloatBuffer.GetData(), InFloatBuffer.Num());
+		TArrayView<float> BufferToSumToView(BufferToSumTo.GetData(), BufferToSumTo.Num());
 
-		const int32 NumOut = OutPowerValues.Num();
-		const int32 NumToSimd = NumOut & MathIntrinsics::SimdMask;
-		const int32 NumNotToSimd = NumOut & MathIntrinsics::NotSimdMask;
+		ArrayMixIn(InFloatBufferView, BufferToSumToView);
+	}
 
-		const float* InComplexData = InComplexValues.GetData();
-		float* OutPowerData = OutPowerValues.GetData();
+	void ArrayMixIn(TArrayView<const float> InFloatBuffer, TArrayView<float> BufferToSumTo, const float StartGain, const float EndGain)
+	{
+		checkf(InFloatBuffer.Num() == BufferToSumTo.Num(), TEXT("Buffers must be equal size"));
+
+		int32 Num = InFloatBuffer.Num();
 
-		for (int32 i = 0; i < NumToSimd; i += 4)
+		if (FMath::IsNearlyEqual(StartGain, EndGain))
+		{
+			// No need to do anything if start and end values are both 0.0
+			if (StartGain == 0.0f)
+			{
+				return;
+			}
+			else
+			{
+				ArrayMixIn(InFloatBuffer, BufferToSumTo, StartGain);
+			}
+		}
+		else
 		{
-			VectorRegister4Float VectorComplex1 = VectorLoadAligned(&InComplexData[2 * i]);
-			VectorRegister4Float VectorSquared1 = VectorMultiply (VectorComplex1, VectorComplex1);
+			if (bAudio_FloatArrayMath_ISPC_Enabled)
+			{
+#if INTEL_ISPC
+				ispc::ArrayMixInWithDelta(InFloatBuffer.GetData(), BufferToSumTo.GetData(), Num, StartGain, EndGain);
+#endif
+			}
+			else
+			{
+				const int32 NumToSimd = Num & MathIntrinsics::SimdMask;
+				const int32 NumNotToSimd = Num & MathIntrinsics::NotSimdMask;
 
-			VectorRegister4Float VectorComplex2 = VectorLoadAligned(&InComplexData[(2 * i) + 4]);
-			VectorRegister4Float VectorSquared2 = VectorMultiply (VectorComplex2, VectorComplex2);
+				if (NumToSimd)
+				{
+					const int32 NumIterations = Num / AUDIO_NUM_FLOATS_PER_VECTOR_REGISTER;
+					const float DeltaValue = ((EndGain - StartGain) / NumIterations);
 
-			VectorRegister4Float VectorSquareReal = VectorShuffle(VectorSquared1, VectorSquared2, 0, 2, 0, 2);
-			VectorRegister4Float VectorSquareImag = VectorShuffle(VectorSquared1, VectorSquared2, 1, 3, 1, 3);
+					VectorRegister4Float Gain = VectorLoadFloat1(&StartGain);
+					VectorRegister4Float Delta = VectorLoadFloat1(&DeltaValue);
 
-			VectorRegister4Float VectorOut = VectorAdd(VectorSquareReal, VectorSquareImag);
-			
-			VectorStoreAligned(VectorOut, &OutPowerData[i]);
-		}
+					for (int32 i = 0; i < NumToSimd; i += AUDIO_NUM_FLOATS_PER_VECTOR_REGISTER)
+					{
+						VectorRegister4Float Input = VectorLoad(&InFloatBuffer[i]);
+						VectorRegister4Float Output = VectorLoad(&BufferToSumTo[i]);
+						Output = VectorMultiplyAdd(Input, Gain, Output);
 
-		if (NumNotToSimd)
-		{
-			TArrayView<const float> ComplexView(&InComplexData[2 * NumToSimd], 2 * NumNotToSimd);
-			TArrayView<float> PowerView(&OutPowerData[NumToSimd], NumNotToSimd);
+						VectorStore(Output, &BufferToSumTo[i]);
+
+						Gain = VectorAdd(Gain, Delta);
+					}
+				}
+
+				if (NumNotToSimd)
+				{
+					const float DeltaValue = ((EndGain - StartGain) / Num);
+					float Gain = (NumToSimd * DeltaValue) + StartGain;
 
-			ArrayComplexToPower(ComplexView, PowerView);
+					for (int32 i = NumToSimd; i < Num; ++i)
+					{
+						BufferToSumTo[i] += InFloatBuffer[i] * Gain;
+						Gain += DeltaValue;
+					}
+				}
+			}
 		}
 	}
 
+	void ArrayMixIn(const FAlignedFloatBuffer& InFloatBuffer, FAlignedFloatBuffer& BufferToSumTo, const float StartGain, const float EndGain)
+	{
+		TArrayView<const float> InFloatBufferView(InFloatBuffer.GetData(), InFloatBuffer.Num());
+		TArrayView<float> BufferToSumToView(BufferToSumTo.GetData(), BufferToSumTo.Num());
+
+		ArrayMixIn(InFloatBufferView, BufferToSumToView, StartGain, EndGain);
+	}
 
 	FContiguousSparse2DKernelTransform::FContiguousSparse2DKernelTransform(const int32 NumInElements, const int32 NumOutElements)
 	:	NumIn(NumInElements)
@@ -1099,9 +2348,19 @@ namespace Audio
 
 			// dot prod 'em. 
 			int32 NumToMult = Row.OffsetValues.Num();
-			for (int32 i = 0; i < NumToMult; i++)
+
+			if (bAudio_FloatArrayMath_ISPC_Enabled)
+			{
+#if INTEL_ISPC
+				ispc::TransformArrayRow(OffsetInData, RowValuePtr, OutArray, RowIndex, NumToMult);
+#endif
+			}
+			else
 			{
-				OutArray[RowIndex] += OffsetInData[i] * RowValuePtr[i];
+				for (int32 i = 0; i < NumToMult; i++)
+				{
+					OutArray[RowIndex] += OffsetInData[i] * RowValuePtr[i];
+				}
 			}
 		}
 	}
diff --git a/Engine/Source/Runtime/SignalProcessing/Private/FloatArrayMath.ispc b/Engine/Source/Runtime/SignalProcessing/Private/FloatArrayMath.ispc
new file mode 100644
--- /dev/null
+++ b/Engine/Source/Runtime/SignalProcessing/Private/FloatArrayMath.ispc
@@ -0,0 +1,666 @@
+// Copyright Epic Games, Inc. All Rights Reserved.
+
+#include "Math/Aossoa.isph"
+#include "Math/Scalar.isph"
+
+export void ArraySum(const uniform float InData[],
+					uniform float &OutSum,
+					const uniform int Num)
+{
+	varying float Sum = 0.0f;
+
+	foreach(i = 0 ... Num)
+	{
+		Sum += InData[i];
+	}
+
+	OutSum = reduce_add(Sum);
+}
+
+export void ArraySum2(const uniform float InFloatBuffer1[],
+						const uniform float InFloatBuffer2[],
+						uniform float OutputBuffer[],
+						const uniform int NumSamples)
+{
+	foreach(i = 0 ... NumSamples)
+	{
+		const varying float Input1 = InFloatBuffer1[i];
+		const varying float Input2 = InFloatBuffer2[i];
+
+		OutputBuffer[i] = Input1 + Input2;
+	}
+}
+
+export void ArrayCumulativeSum(const uniform float InViewPtr[],
+								uniform float OutDataPtr[],
+								const uniform int Num)
+{
+	uniform float Sum = 0;
+
+	for(uniform int32 i = 0; i < Num; i++)
+	{
+		Sum += InViewPtr[i];
+		OutDataPtr[i] = Sum;
+	}
+}
+
+export void ArrayMean(const uniform float DataPtr[],
+					uniform float &OutMean,
+					const uniform int Num)
+{
+	varying float Sum = 0.0f;
+
+	foreach(i = 0 ... Num)
+	{
+		Sum += DataPtr[i];
+	}
+
+	OutMean = reduce_add(Sum) / (uniform float)Num;
+}
+
+export void ArrayMeanSquared(const uniform float DataPtr[],
+							uniform float &OutMean,
+							const uniform int Num)
+{
+	varying float Sum = 0.0f;
+
+	foreach(i = 0 ... Num)
+	{
+		Sum += DataPtr[i] * DataPtr[i];
+	}
+
+	OutMean = reduce_add(Sum) / (uniform float)Num;
+}
+
+export uniform float ArrayGetMagnitude(const uniform float Buffer[], const uniform int NumSamples)
+{
+	varying float Sum = 0.0f;
+	foreach(i = 0 ... NumSamples)
+	{
+		const float Input = Buffer[i];
+		Sum += Input * Input;
+	}
+	return sqrt(reduce_add(Sum));
+}
+
+export uniform float ArrayGetAverageValue(const uniform float Buffer[], const uniform int NumSamples)
+{
+	varying float Sum = 0.0f;
+	foreach(i = 0 ... NumSamples)
+	{
+		const float Input = Buffer[i];
+		Sum += Input;
+	}
+	return reduce_add(Sum) / NumSamples;
+}
+
+export uniform float ArrayGetAverageAbsValue(const uniform float Buffer[], const uniform int NumSamples)
+{
+	varying float Sum = 0.0f;
+	foreach(i = 0 ... NumSamples)
+	{
+		const float Input = Buffer[i];
+		Sum += abs(Input);
+	}
+	return reduce_add(Sum) / NumSamples;
+}
+
+export void ArrayMeanFilter(const uniform float SummedDataPtr[],
+							const uniform int WindowSize,
+							const uniform int WindowOrigin,
+							uniform float OutDataPtr[],
+							const uniform float LastSummedData,
+							const uniform int Num)
+{
+	const uniform int32 LastIndexBeforeEndBoundaryCondition = max(WindowOrigin + 1, Num - WindowSize + WindowOrigin + 1);
+	const uniform int32 StartOffset = -WindowOrigin - 1;
+	const uniform int32 EndOffset = WindowSize - WindowOrigin - 1;
+	const uniform int32 WindowTail = WindowSize - WindowOrigin;
+
+	if ((WindowSize - WindowOrigin) < Num)
+	{
+		// Handle boundary condition where analysis window precedes beginning of array.
+		foreach (i = 0 ... (WindowOrigin + 1))
+		{
+			OutDataPtr[i] = SummedDataPtr[i + EndOffset] / max(1.f, (varying float)(WindowTail + i));
+		}
+
+		// No boundary conditions to handle here.
+		const uniform float MeanDivisor = (uniform float)(WindowSize);
+		foreach (i = WindowOrigin + 1 ... LastIndexBeforeEndBoundaryCondition)
+		{
+			OutDataPtr[i] = (SummedDataPtr[i + EndOffset] - SummedDataPtr[i + StartOffset]) / MeanDivisor;
+		}
+	}
+	else
+	{
+		// Handle boundary condition where window precedes beginning and goes past end of array
+		const uniform float ArrayMean = LastSummedData / (uniform float)(Num);
+		foreach (i = 0 ... LastIndexBeforeEndBoundaryCondition)
+		{
+			OutDataPtr[i] = ArrayMean;
+		}
+	}
+
+	// Handle boundary condition where analysis window goes past end of array.
+	foreach (i = LastIndexBeforeEndBoundaryCondition ... Num)
+	{
+		OutDataPtr[i] = (LastSummedData - SummedDataPtr[i + StartOffset]) / (varying float)(Num - i + WindowOrigin);
+	}
+}
+
+export void ArrayGetEuclideanNorm(const uniform float InViewData[],
+								uniform float &OutEuclideanNorm,
+								const uniform int Num)
+{
+	varying float SquareSum = 0.0f;
+
+	foreach(i = 0 ... Num)
+	{
+		SquareSum += InViewData[i] * InViewData[i];
+	}
+
+	OutEuclideanNorm = sqrt(reduce_add(SquareSum));
+}
+
+export void ArrayAbs(const uniform float InData[],
+					uniform float OutData[],
+					const uniform int Num)
+{
+	foreach(i = 0 ... Num)
+	{
+		OutData[i] = abs(InData[i]);
+	}
+}
+
+export void ArrayAbsInPlace(uniform float Data[],
+							const uniform int Num)
+{
+	foreach(i = 0 ... Num)
+	{
+		Data[i] = abs(Data[i]);
+	}
+}
+
+export void ArrayClampMinInPlace(uniform float Data[],
+								const uniform float InMin,
+								const uniform int Num)
+{
+	foreach(i = 0 ... Num)
+	{
+		Data[i] = max(InMin, Data[i]);
+	}
+}
+
+export void ArrayClampMaxInPlace(uniform float Data[],
+								const uniform float InMax,
+								const uniform int Num)
+{
+	foreach(i = 0 ... Num)
+	{
+		Data[i] = min(InMax, Data[i]);
+	}
+}
+
+export void ArrayClampInPlace(uniform float Data[],
+								const uniform float InMin,
+								const uniform float InMax,
+								const uniform int Num)
+{
+	foreach(i = 0 ... Num)
+	{
+		Data[i] = clamp(Data[i], InMin, InMax);
+	}
+}
+
+export void ArrayMinMaxNormalize(const uniform float InDataPtr[],
+								uniform float OutDataPtr[],
+								const uniform int Num)
+{
+	varying float VMaxValue = InDataPtr[0];
+	varying float VMinValue = InDataPtr[0];
+
+	// Determine min and max
+	foreach(i = 1 ... Num)
+	{
+		const varying float Data = InDataPtr[i];
+		VMaxValue = max(VMaxValue, Data);
+		VMinValue = min(VMinValue, Data);
+	}
+
+	const uniform float MaxValue = reduce_max(VMaxValue);
+	const uniform float MinValue = reduce_min(VMinValue);
+
+	// Normalize data by subtracting minimum value and dividing by range
+	const uniform float Scale = 1.f / max(FLOAT_SMALL_NUMBER, MaxValue - MinValue);
+	foreach(i = 0 ... Num)
+	{
+		OutDataPtr[i] = (InDataPtr[i] - MinValue) * Scale;
+	}
+}
+
+export void ArrayMultiplyInPlace(const uniform float InFloatBuffer[],
+									uniform float BufferToMultiply[],
+									const uniform int NumSamples)
+{
+	foreach(i = 0 ... NumSamples)
+	{
+		const varying float Input = InFloatBuffer[i];
+		const varying float Output = BufferToMultiply[i];
+
+		BufferToMultiply[i] = Input * Output;
+	}
+}
+
+export void ArrayComplexMultiplyInPlace(const uniform float InData1[],
+										uniform float InData2[],
+										const uniform int Num)
+{
+#if TARGET_WIDTH == 4
+	static const varying float RealSignFlip = { -1.f, 1.f, -1.f, 1.f };
+#elif TARGET_WIDTH == 8
+	static const varying float RealSignFlip = { -1.f, 1.f, -1.f, 1.f, -1.f, 1.f, -1.f, 1.f };
+#elif TARGET_WIDTH == 16
+	static const varying float RealSignFlip = { -1.f, 1.f, -1.f, 1.f, -1.f, 1.f, -1.f, 1.f, -1.f, 1.f, -1.f, 1.f, -1.f, 1.f, -1.f, 1.f };
+#endif
+
+	uniform int NumToSimd = Num & ~(programCount-1);
+
+	for(uniform int i = 0; i < NumToSimd; i+=programCount)
+	{
+		const varying float VectorData1 = InData1[i+programIndex];
+		const varying float VectorData2 = InData2[i+programIndex];
+
+		const varying float VectorData1Real = VectorSwizzle(VectorData1, 0, 0, 2, 2);
+		const varying float VectorData1Imag = VectorSwizzle(VectorData1, 1, 1, 3, 3);
+		const varying float VectorData2Swizzle = VectorSwizzle(VectorData2, 1, 0, 3, 2);
+
+		varying float Result = VectorData1Imag * VectorData2Swizzle;
+		Result = Result * RealSignFlip;
+		Result = (VectorData1Real * VectorData2) + Result;
+
+		InData2[i+programIndex] = Result;
+	}
+
+	for(uniform int i = NumToSimd; i < Num; i+=2)
+	{
+		uniform float Real = (InData1[i] * InData2[i]) - (InData1[i + 1] * InData2[i + 1]);
+		uniform float Imag = (InData1[i] * InData2[i + 1]) + (InData1[i + 1] * InData2[i]);
+		InData2[i] = Real;
+		InData2[i + 1] = Imag;
+	}
+}
+
+export void ArrayMultiplyByConstant(const uniform float InFloatBuffer[],
+									const uniform float InValue,
+									uniform float OutFloatBuffer[],
+									const uniform int InNumSamples)
+{
+	foreach(i = 0 ... InNumSamples)
+	{
+		OutFloatBuffer[i] = InValue * InFloatBuffer[i];
+	}
+}
+
+export void ArrayMultiplyByConstantInPlace(uniform float OutFloatBuffer[],
+										const uniform int NumSamples,
+										const uniform float InGain)
+{
+	foreach(i = 0 ... NumSamples)
+	{
+		const varying float Output = OutFloatBuffer[i];
+		OutFloatBuffer[i] = Output * InGain;
+	}
+}
+
+export void ArrayAddInPlace(const uniform float InData[],
+							uniform float InAccumulateData[],
+							const uniform int Num)
+{
+	foreach(i = 0 ... Num)
+	{
+		InAccumulateData[i] += InData[i];
+	}
+}
+
+export void ArrayAddConstantInplace(uniform float OutFloatBuffer[],
+									const uniform int NumSamples,
+									const uniform float InConstant)
+{
+	foreach(i = 0 ... NumSamples)
+	{
+		const varying float Output = OutFloatBuffer[i];
+		OutFloatBuffer[i] = Output + InConstant;
+	}
+}
+
+export void ArrayMultiplyAddInPlace(const uniform float InData[],
+									const uniform float InMultiplier,
+									uniform float InAccumulateData[],
+									const uniform int Num)
+{
+	foreach(i = 0 ... Num)
+	{
+		InAccumulateData[i] += InData[i] * InMultiplier;
+	}
+}
+
+export void ArrayLerpAddInPlace(const uniform float InData[],
+								const uniform float InStartMultiplier,
+								const uniform float InEndMultiplier,
+								uniform float InAccumulateData[],
+								const uniform int Num)
+{
+	const uniform float Delta = (InEndMultiplier - InStartMultiplier) / max(1.f, (uniform float)(Num - 1));
+
+	foreach(i = 0 ... Num)
+	{
+		const varying float Multiplier = InStartMultiplier + (i * Delta);
+		InAccumulateData[i] += InData[i] * Multiplier;
+	}
+}
+
+export void ArraySubtract(const uniform float InMinuend[],
+								const uniform float InSubtrahend[],
+								uniform float OutBuffer[],
+								const uniform int InNum)
+{
+	foreach(i = 0 ... InNum)
+	{
+		const varying float Input1 = InMinuend[i];
+		const varying float Input2 = InSubtrahend[i];
+
+		const varying float Output = Input1 - Input2;
+
+		OutBuffer[i] = Output;
+	}
+}
+
+export void ArraySubtractInPlace1(const uniform float InMinuend[],
+									uniform float InOutSubtrahend[],
+									const uniform int InNum)
+{
+	foreach(i = 0 ... InNum)
+	{
+		const varying float Input1 = InMinuend[i];
+		const varying float Input2 = InOutSubtrahend[i];
+
+		const varying float Output = Input1 - Input2;
+
+		InOutSubtrahend[i] = Output;
+	}
+}
+
+export void ArraySubtractInPlace2(uniform float InOutMinuend[],
+									const uniform float InSubtrahend[],
+									const uniform int InNum)
+{
+	foreach(i = 0 ... InNum)
+	{
+		const varying float Input1 = InOutMinuend[i];
+		const varying float Input2 = InSubtrahend[i];
+
+		const varying float Output = Input1 - Input2;
+
+		InOutMinuend[i] = Output;
+	}
+}
+
+export void ArraySubtractByConstantInPlace(uniform float InData[],
+											const uniform float InSubtrahend,
+											const uniform int Num)
+{
+	foreach(i = 0 ... Num)
+	{
+		InData[i] -= InSubtrahend;
+	}
+}
+
+export void ArraySquare(const uniform float InData[],
+						uniform float OutData[],
+						const uniform int Num)
+{
+	foreach(i = 0 ... Num)
+	{
+		OutData[i] = InData[i] * InData[i];
+	}
+}
+
+export void ArraySquareInPlace(uniform float InData[],
+								const uniform int Num)
+{
+	foreach(i = 0 ... Num)
+	{
+		InData[i] = InData[i] * InData[i];
+	}
+}
+
+export void ArraySqrtInPlace(uniform float InData[],
+								const uniform int Num)
+{
+	foreach(i = 0 ... Num)
+	{
+		InData[i] = sqrt(InData[i]);
+	}
+}
+
+export void ArrayComplexConjugate(const uniform float InData[],
+									uniform float OutData[],
+									const uniform int Num)
+{
+	foreach(i = 0 ... Num)
+	{
+		OutData[i] = select(i & 1, -InData[i], InData[i]);
+	}
+}
+
+export void ArrayComplexConjugateInPlace(uniform float InData[],
+											const uniform int Num)
+{
+	foreach(i = 0 ... Num)
+	{
+		InData[i] = select(i & 1, -InData[i], InData[i]);
+	}
+}
+
+export void ArrayMagnitudeToDecibelInPlace(uniform float InData[],
+											const uniform float InMinimumDb,
+											const uniform int Num)
+{
+	const uniform float Loge10 = log(10.0f);
+	const uniform float Scale = 20.f / Loge10;
+	const uniform float Minimum = exp(InMinimumDb * Loge10 / 20.f);
+
+	foreach(i = 0 ... Num)
+	{
+		varying float VectorData = InData[i];
+		VectorData = max(VectorData, Minimum);
+		InData[i] = log(VectorData) * Scale;
+	}
+}
+
+export void ArrayPowerToDecibelInPlace(uniform float InData[],
+										const uniform float InMinimumDb,
+										const uniform int Num)
+{
+	const uniform float Loge10 = log(10.0f);
+	const uniform float Scale = 10.f / Loge10;
+	const uniform float Minimum = exp(InMinimumDb * Loge10 / 10.f);
+
+	foreach(i = 0 ... Num)
+	{
+		varying float VectorData = InData[i];
+		VectorData = max(VectorData, Minimum);
+		InData[i] = log(VectorData) * Scale;
+	}
+}
+
+export void ArrayComplexToPowerInterleaved(const uniform float InComplexData[],
+										uniform float OutPowerData[],
+										const uniform int Num)
+{
+	uniform int NumToSimd = Num & ~(programCount-1);
+
+	unmasked
+	{
+		for(uniform int i = 0; i < NumToSimd; i+=programCount)
+		{
+			const uniform int ComplexPos = 2 * i;
+			varying float RealValue, ImagValue;
+			aos_to_soa2_ispc((uniform float *uniform)&InComplexData[ComplexPos], &RealValue, &ImagValue);
+
+			OutPowerData[i + programIndex] = (RealValue * RealValue) + (ImagValue * ImagValue);
+		}
+	}
+
+	for (uniform int32 i = NumToSimd; i < Num; i++)
+	{
+		uniform int ComplexPos = 2 * i;
+
+		uniform float RealValue = InComplexData[ComplexPos];
+		uniform float ImagValue = InComplexData[ComplexPos + 1];
+
+		OutPowerData[i] = (RealValue * RealValue) + (ImagValue * ImagValue);
+	}
+}
+
+export void ArrayComplexToPower(const uniform float InRealSamples[],
+									const uniform float InImaginarySamples[],
+									uniform float OutPowerSamples[],
+									const uniform int NumSamples)
+{
+	foreach(i = 0 ... NumSamples)
+	{
+		const varying float Real = InRealSamples[i];
+		const varying float Imag = InImaginarySamples[i];
+
+		OutPowerSamples[i] = (Real * Real) + (Imag * Imag);
+	}
+}
+
+export void ArrayUnderflowClamp(uniform float InOutBuffer[],
+									const uniform int InNum)
+{
+	foreach(i = 0 ... InNum)
+	{
+		const varying float VInOut = InOutBuffer[i];
+
+		// Create mask of denormal numbers.
+		const varying float Mask = VectorBitwiseAnd(VectorCompareGT(VInOut, -FLT_MIN), VectorCompareLT(VInOut, FLT_MIN));
+
+		// Choose between zero or original number based upon mask.
+		InOutBuffer[i] = VectorSelect(Mask, 0.0f, VInOut);
+	}
+}
+
+export void ArrayRangeClamp(uniform float InOutBuffer[],
+								const uniform int InNum,
+								const uniform float InMinValue,
+								const uniform float InMaxValue)
+{
+	foreach(i = 0 ... InNum)
+	{
+		InOutBuffer[i] = clamp(InOutBuffer[i], InMinValue, InMaxValue);
+	}
+}
+
+export void ArraySetToConstantInplace(uniform float InBuffer[],
+									const uniform int NumSamples,
+									const uniform float InConstant)
+{
+	foreach(i = 0 ... NumSamples)
+	{
+		InBuffer[i] = InConstant;
+	}
+}
+
+export void ArrayWeightedSumTwoGain(const uniform float InBuffer1[],
+										const uniform float InGain1,
+										const uniform float InBuffer2[],
+										const uniform float InGain2,
+										uniform float OutBuffer[],
+										const uniform int InNum)
+{
+	foreach(i = 0 ... InNum)
+	{
+		OutBuffer[i] = (InBuffer1[i] * InGain1) + (InBuffer2[i] * InGain2);
+	}
+}
+
+export void ArrayWeightedSumOneGain(const uniform float InBuffer1[],
+										const uniform float InGain1,
+										const uniform float InBuffer2[],
+										uniform float OutBuffer[],
+										const uniform int InNum)
+{
+	foreach(i = 0 ... InNum)
+	{
+		OutBuffer[i] = (InBuffer1[i] * InGain1) + InBuffer2[i];
+	}
+}
+
+export void ArrayFade(uniform float OutFloatBuffer[],
+							const uniform int NumSamples,
+							const uniform float StartValue,
+							const uniform float EndValue)
+{
+	const uniform float DeltaValue = ((EndValue - StartValue) / NumSamples);
+
+	foreach(i = 0 ... NumSamples)
+	{
+		const varying float Gain = (i * DeltaValue) + StartValue;
+		OutFloatBuffer[i] = OutFloatBuffer[i] * Gain;
+	}
+}
+
+export void ArrayMixInWithGain(const uniform float InFloatBuffer[],
+								uniform float BufferToSumTo[],
+								const uniform int NumSamples,
+								const uniform float Gain)
+{
+	foreach(i = 0 ... NumSamples)
+	{
+		BufferToSumTo[i] += InFloatBuffer[i] * Gain;
+	}
+}
+
+export void ArrayMixIn(const uniform float InFloatBuffer[],
+							uniform float BufferToSumTo[],
+							const uniform int NumSamples)
+{
+	foreach(i = 0 ... NumSamples)
+	{
+		BufferToSumTo[i] += InFloatBuffer[i];
+	}
+}
+
+export void ArrayMixInWithDelta(const uniform float InFloatBuffer[],
+									uniform float BufferToSumTo[],
+									const uniform int NumSamples,
+									const uniform float StartGain,
+									const uniform float EndGain)
+{
+	const uniform float DeltaValue = ((EndGain - StartGain) / NumSamples);
+
+	foreach(i = 0 ... NumSamples)
+	{
+		const varying float Gain = (i * DeltaValue) + StartGain;
+		BufferToSumTo[i] += InFloatBuffer[i] * Gain;
+	}
+}
+
+export void TransformArrayRow(const uniform float OffsetInData[],
+								const uniform float *uniform RowValuePtr,
+								uniform float OutArray[],
+								const uniform int RowIndex,
+								const uniform int NumToMult)
+{
+	varying float RowSum = 0.0f;
+
+	foreach(i = 0 ... NumToMult)
+	{
+		RowSum += OffsetInData[i] * RowValuePtr[i];
+	}
+
+	OutArray[RowIndex] += reduce_add(RowSum);
+}
diff --git a/Engine/Source/Runtime/SignalProcessing/Private/LateReflectionsFast.cpp b/Engine/Source/Runtime/SignalProcessing/Private/LateReflectionsFast.cpp
--- a/Engine/Source/Runtime/SignalProcessing/Private/LateReflectionsFast.cpp
+++ b/Engine/Source/Runtime/SignalProcessing/Private/LateReflectionsFast.cpp
@@ -189,7 +189,7 @@ void FLateReflectionsPlate::ProcessAudioFrames(
 	WorkBufferC.AddUninitialized(InNum);
 
 	// Copy Feedback Samples to aligned internal buffer.
-	BufferWeightedSumFast(InFeedbackSamples, 1.0 - Decay, InSamples, WorkBufferA);
+	ArrayWeightedSum(InFeedbackSamples, 1.0 - Decay, InSamples, WorkBufferA);
 
 	// Input -> ModulatedAPF
 	ModulatedAPF->ProcessAudio(WorkBufferA, InDelayModulations, WorkBufferB);
@@ -201,13 +201,13 @@ void FLateReflectionsPlate::ProcessAudioFrames(
 	DelayD->ProcessAudio(OutPlateSamples.Taps[2], WorkBufferA);
 
 	// Apply dampening
-	MultiplyBufferByConstantInPlace(WorkBufferA, 1.0f - Dampening);
+	ArrayMultiplyByConstantInPlace(WorkBufferA, 1.0f - Dampening);
 	
 	// Delay1 -> LPF
 	LPF->ProcessAudio(WorkBufferA, WorkBufferB);
 
 	// Apply decay
-	MultiplyBufferByConstantInPlace(WorkBufferB, 1.0f - Decay);
+	ArrayMultiplyByConstantInPlace(WorkBufferB, 1.0f - Decay);
 
 	// LPF -> APF
 	APF->ProcessAudio(WorkBufferB, WorkBufferA, WorkBufferC);
@@ -447,7 +447,7 @@ void FLateReflectionsFast::ProcessAudioBuffer(const float* InSampleData, const i
 	}
 
 	// Apply bandwidth, gain and channel averaging multipliers
-	MultiplyBufferByConstantInPlace(WorkBufferA, Settings.Bandwidth * Gain / InNumChannels);
+	ArrayMultiplyByConstantInPlace(WorkBufferA, Settings.Bandwidth * Gain / InNumChannels);
 	
 	// Predelay
 	PreDelay->ProcessAudio(WorkBufferA, WorkBufferB);
@@ -472,6 +472,23 @@ void FLateReflectionsFast::ProcessAudioBuffer(const float* InSampleData, const i
 	LeftPlate->ProcessAudioFrames(WorkBufferA, WorkBufferC, LeftDelayModSamples, LeftPlateOutputs);
 	RightPlate->ProcessAudioFrames(WorkBufferA, WorkBufferB, RightDelayModSamples, RightPlateOutputs);
 
+	TArrayView<float> OutLeftSampleDataView(OutLeftSampleData, InNumFrames);
+	TArrayView<float> OutRightSampleDataView(OutRightSampleData, InNumFrames);
+
+	TArrayView<const float> LeftTaps1View(LeftPlateOutputs.Taps[1].GetData(), InNumFrames);
+	TArrayView<const float> LeftTaps2View(LeftPlateOutputs.Taps[2].GetData(), InNumFrames);
+	TArrayView<const float> LeftTaps3View(LeftPlateOutputs.Taps[3].GetData(), InNumFrames);
+	TArrayView<const float> LeftTaps4View(LeftPlateOutputs.Taps[4].GetData(), InNumFrames);
+	TArrayView<const float> LeftTaps5View(LeftPlateOutputs.Taps[5].GetData(), InNumFrames);
+	TArrayView<const float> LeftTaps6View(LeftPlateOutputs.Taps[6].GetData(), InNumFrames);
+
+	TArrayView<const float> RightTaps1View(RightPlateOutputs.Taps[1].GetData(), InNumFrames);
+	TArrayView<const float> RightTaps2View(RightPlateOutputs.Taps[2].GetData(), InNumFrames);
+	TArrayView<const float> RightTaps3View(RightPlateOutputs.Taps[3].GetData(), InNumFrames);
+	TArrayView<const float> RightTaps4View(RightPlateOutputs.Taps[4].GetData(), InNumFrames);
+	TArrayView<const float> RightTaps5View(RightPlateOutputs.Taps[5].GetData(), InNumFrames);
+	TArrayView<const float> RightTaps6View(RightPlateOutputs.Taps[6].GetData(), InNumFrames);
+
 	// Left Output
 	// Channel -> [Plate and Tap#]
 	// L -> +R0
@@ -482,12 +499,12 @@ void FLateReflectionsFast::ProcessAudioBuffer(const float* InSampleData, const i
 	// L -> -L3
 	// L -> -L5
 	FMemory::Memcpy(OutLeftSampleData, RightPlateOutputs.Taps[0].GetData(), InNumFrames * sizeof(float));
-	MixInBufferFast(RightPlateOutputs.Taps[2].GetData(), OutLeftSampleData, InNumFrames);
-	BufferSubtractInPlace2Fast(OutLeftSampleData, RightPlateOutputs.Taps[4].GetData(), InNumFrames);
-	MixInBufferFast(RightPlateOutputs.Taps[6].GetData(), OutLeftSampleData, InNumFrames);
-	BufferSubtractInPlace2Fast(OutLeftSampleData, LeftPlateOutputs.Taps[1].GetData(), InNumFrames);
-	BufferSubtractInPlace2Fast(OutLeftSampleData, LeftPlateOutputs.Taps[3].GetData(), InNumFrames);
-	BufferSubtractInPlace2Fast(OutLeftSampleData, LeftPlateOutputs.Taps[5].GetData(), InNumFrames);
+	ArrayMixIn(RightTaps2View, OutLeftSampleDataView);
+	ArraySubtractInPlace2(OutLeftSampleDataView, RightTaps4View);
+	ArrayMixIn(RightTaps6View, OutLeftSampleDataView);
+	ArraySubtractInPlace2(OutLeftSampleDataView, LeftTaps1View);
+	ArraySubtractInPlace2(OutLeftSampleDataView, LeftTaps3View);
+	ArraySubtractInPlace2(OutLeftSampleDataView, LeftTaps5View);
 	
 	// Right Output
 	// Channel -> [Plate and Tap#]
@@ -499,12 +516,12 @@ void FLateReflectionsFast::ProcessAudioBuffer(const float* InSampleData, const i
 	// R -> -R3
 	// R -> -R5
 	FMemory::Memcpy(OutRightSampleData, LeftPlateOutputs.Taps[0].GetData(), InNumFrames * sizeof(float));
-	MixInBufferFast(LeftPlateOutputs.Taps[2].GetData(), OutRightSampleData, InNumFrames);
-	BufferSubtractInPlace2Fast(OutRightSampleData, LeftPlateOutputs.Taps[4].GetData(), InNumFrames);
-	MixInBufferFast(LeftPlateOutputs.Taps[6].GetData(), OutRightSampleData, InNumFrames);
-	BufferSubtractInPlace2Fast(OutRightSampleData, RightPlateOutputs.Taps[1].GetData(), InNumFrames);
-	BufferSubtractInPlace2Fast(OutRightSampleData, RightPlateOutputs.Taps[3].GetData(), InNumFrames);
-	BufferSubtractInPlace2Fast(OutRightSampleData, RightPlateOutputs.Taps[5].GetData(), InNumFrames);
+	ArrayMixIn(LeftTaps2View, OutRightSampleDataView);
+	ArraySubtractInPlace2(OutRightSampleDataView, LeftTaps4View);
+	ArrayMixIn(LeftTaps6View, OutRightSampleDataView);
+	ArraySubtractInPlace2(OutRightSampleDataView, RightTaps1View);
+	ArraySubtractInPlace2(OutRightSampleDataView, LeftTaps3View);
+	ArraySubtractInPlace2(OutRightSampleDataView, LeftTaps5View);
 }
 
 void FLateReflectionsFast::ApplySettings()
diff --git a/Engine/Source/Runtime/SignalProcessing/Private/LinkwitzRileyBandSplitter.cpp b/Engine/Source/Runtime/SignalProcessing/Private/LinkwitzRileyBandSplitter.cpp
--- a/Engine/Source/Runtime/SignalProcessing/Private/LinkwitzRileyBandSplitter.cpp
+++ b/Engine/Source/Runtime/SignalProcessing/Private/LinkwitzRileyBandSplitter.cpp
@@ -178,7 +178,8 @@ namespace Audio
 			constexpr int32 IsOddBitMask = 0x00000001;
 			if ((static_cast<int32>(FilterOrder) & IsOddBitMask) && (BandId & IsOddBitMask))
 			{
-				MultiplyBufferByConstantInPlace(BandBufferPtr, NumSamples, -1.f);
+				TArrayView<float> BandBufferView(BandBufferPtr, NumSamples);
+				ArrayMultiplyByConstantInPlace(BandBufferView, -1.f);
 			}
 
 			CopyToBuffer(OutBuffer[BandId], BandBufferPtr, NumSamples);
diff --git a/Engine/Source/Runtime/SignalProcessing/Private/MultithreadedPatching.cpp b/Engine/Source/Runtime/SignalProcessing/Private/MultithreadedPatching.cpp
--- a/Engine/Source/Runtime/SignalProcessing/Private/MultithreadedPatching.cpp
+++ b/Engine/Source/Runtime/SignalProcessing/Private/MultithreadedPatching.cpp
@@ -50,16 +50,17 @@ namespace Audio
 		}
 
 		int32 PopResult = InternalBuffer.Pop(OutBuffer, NumSamples);
+		TArrayView<float> OutBufferView(OutBuffer, PopResult);
 
 		// Apply gain stage.
 		float TG = TargetGain, PG = PreviousGain;
 		if (FMath::IsNearlyEqual(TG, PG))
 		{
-			MultiplyBufferByConstantInPlace(OutBuffer, PopResult, PreviousGain);
+			ArrayMultiplyByConstantInPlace(OutBufferView, PreviousGain);
 		}
 		else
 		{
-			FadeBufferFast(OutBuffer, PopResult, PreviousGain, TargetGain);
+			ArrayFade(OutBufferView, PreviousGain, TargetGain);
 			PreviousGain = TargetGain;
 		}
 		
@@ -91,14 +92,17 @@ namespace Audio
 			PopResult = InternalBuffer.Pop(MixingBuffer.GetData(), NumSamples);
 		}
 
+		TArrayView<const float> MixingBufferView(MixingBuffer.GetData(), PopResult);
+		TArrayView<float> OutBufferView(OutBuffer, PopResult);
+
 		float TG = TargetGain, PG = PreviousGain;
 		if (FMath::IsNearlyEqual(TG, PG))
 		{
-			MixInBufferFast(MixingBuffer.GetData(), OutBuffer, PopResult, PreviousGain);
+			ArrayMixIn(MixingBufferView, OutBufferView, PreviousGain);
 		}
 		else
 		{
-			MixInBufferFast(MixingBuffer.GetData(), OutBuffer, PopResult, PreviousGain, TargetGain);
+			ArrayMixIn(MixingBufferView, OutBufferView, PreviousGain, TargetGain);
 			PreviousGain = TargetGain;
 		}
 
diff --git a/Engine/Source/Runtime/SignalProcessing/Private/ReverbFast.cpp b/Engine/Source/Runtime/SignalProcessing/Private/ReverbFast.cpp
--- a/Engine/Source/Runtime/SignalProcessing/Private/ReverbFast.cpp
+++ b/Engine/Source/Runtime/SignalProcessing/Private/ReverbFast.cpp
@@ -139,8 +139,8 @@ namespace Audio {
 			EarlyReflections.ProcessAudio(ScaledInputBuffer, InNumChannels, FrontLeftEarlyReflectionsSamples, FrontRightEarlyReflectionsSamples);
 			LateReflections.ProcessAudio(ScaledInputBuffer, InNumChannels, FrontLeftLateReflectionsSamples, FrontRightLateReflectionsSamples);
 			// Add early and late reflections together.
-			SumBuffers(FrontLeftEarlyReflectionsSamples, FrontLeftLateReflectionsSamples, FrontLeftReverbSamples);
-			SumBuffers(FrontRightEarlyReflectionsSamples, FrontRightLateReflectionsSamples, FrontRightReverbSamples);
+			ArraySum(FrontLeftEarlyReflectionsSamples, FrontLeftLateReflectionsSamples, FrontLeftReverbSamples);
+			ArraySum(FrontRightEarlyReflectionsSamples, FrontRightLateReflectionsSamples, FrontRightReverbSamples);
 		}
 
 
@@ -203,8 +203,8 @@ namespace Audio {
 				RightAttenuatedSamples.AddUninitialized(InNumFrames);
 
 				// Reduce volume of output reverbs.
-				BufferMultiplyByConstant(FrontLeftReverbSamples, 0.5f, LeftAttenuatedSamples);
-				BufferMultiplyByConstant(FrontRightReverbSamples, 0.5f, RightAttenuatedSamples);
+				ArrayMultiplyByConstant(FrontLeftReverbSamples, 0.5f, LeftAttenuatedSamples);
+				ArrayMultiplyByConstant(FrontRightReverbSamples, 0.5f, RightAttenuatedSamples);
 
 				const float* FrontLeftSampleData = LeftAttenuatedSamples.GetData();
 				const float* FrontRightSampleData = RightAttenuatedSamples.GetData();
diff --git a/Engine/Source/Runtime/SignalProcessing/Private/RingModulation.cpp b/Engine/Source/Runtime/SignalProcessing/Private/RingModulation.cpp
--- a/Engine/Source/Runtime/SignalProcessing/Private/RingModulation.cpp
+++ b/Engine/Source/Runtime/SignalProcessing/Private/RingModulation.cpp
@@ -88,11 +88,15 @@ namespace Audio
 			}
 		}
 
+		TArrayView<const float> InBufferView(InBuffer, InNumSamples);
+		TArrayView<float> ModulationBufferView(ModulationBuffer.GetData(), InNumSamples);
+		TArrayView<float> OutBufferView(OutBuffer, InNumSamples);
+
 		// Multiply the input buffer by the modulation buffer in-place
-		Audio::MultiplyBuffersInPlace(InBuffer, ModulationBuffer.GetData(), InNumSamples);
+		Audio::ArrayMultiplyInPlace(InBufferView, ModulationBufferView);
 
 		// Perform a buffer weighted sum of the modulation buffer with the Scale value and the dry level as weights
-		Audio::BufferWeightedSumFast(InBuffer, DryLevel, ModulationBuffer.GetData(), Scale, OutBuffer, InNumSamples);
+		Audio::ArrayWeightedSum(InBufferView, DryLevel, ModulationBufferView, Scale, OutBufferView);
 	}
 
 
diff --git a/Engine/Source/Runtime/SignalProcessing/Private/VoiceProcessing.cpp b/Engine/Source/Runtime/SignalProcessing/Private/VoiceProcessing.cpp
--- a/Engine/Source/Runtime/SignalProcessing/Private/VoiceProcessing.cpp
+++ b/Engine/Source/Runtime/SignalProcessing/Private/VoiceProcessing.cpp
@@ -107,7 +107,8 @@ namespace Audio
 				const int32 NumSamplesToFadeOutOver = FMath::Min(NumSamples, DefaultNumSamplesToFadeOutOver);
 
 				const int32 Offset = NumSamples - NumSamplesToFadeOutOver;
-				Audio::FadeBufferFast(&OutAudio[Offset], NumSamplesToFadeOutOver, 1.0f, 0.0f);
+				TArrayView<float> OutAudioView(&OutAudio[Offset], NumSamplesToFadeOutOver);
+				Audio::ArrayFade(OutAudioView, 1.0f, 0.0f);
 			}
 
 			return NumSamples;
@@ -167,7 +168,8 @@ namespace Audio
 		float PeakDetectorOutput = 0.0f; // unused
 		const float EstimatedPeak = PeakDetector.ProcessInput(InAmplitude, PeakDetectorOutput);
 		const float TargetGain = GetTargetGain(EstimatedPeak);
-		Audio::FadeBufferFast(InAudio, NumSamples, PreviousGain, TargetGain);
+		TArrayView<float> InAudioView(InAudio, NumSamples);
+		Audio::ArrayFade(InAudioView, PreviousGain, TargetGain);
 		PreviousGain = TargetGain;
 
 		return TargetGain;
diff --git a/Engine/Source/Runtime/SignalProcessing/Public/DSP/AudioDebuggingUtilities.h b/Engine/Source/Runtime/SignalProcessing/Public/DSP/AudioDebuggingUtilities.h
--- a/Engine/Source/Runtime/SignalProcessing/Public/DSP/AudioDebuggingUtilities.h
+++ b/Engine/Source/Runtime/SignalProcessing/Public/DSP/AudioDebuggingUtilities.h
@@ -12,7 +12,8 @@ void BreakWhenAudible(float* InBuffer, int32 NumSamples)
 {
 	static const float AudibilityThreshold = Audio::ConvertToLinear(-40.0f);
 
-	float BufferAmplitude = Audio::BufferGetAverageAbsValue(InBuffer, NumSamples);
+	TArrayView<const float> InBufferView(InBuffer, NumSamples);
+	float BufferAmplitude = Audio::ArrayGetAverageAbsValue(InBufferView);
 
 	if (BufferAmplitude > AudibilityThreshold)
 	{
@@ -24,7 +25,8 @@ void BreakWhenTooLoud(float* InBuffer, int32 NumSamples)
 {
 	static const float PainThreshold = Audio::ConvertToLinear(3.0f);
 
-	float BufferAmplitude = Audio::BufferGetAverageAbsValue(InBuffer, NumSamples);
+	TArrayView<const float> InBufferView(InBuffer, NumSamples);
+	float BufferAmplitude = Audio::ArrayGetAverageAbsValue(InBufferView);
 
 	if (BufferAmplitude > PainThreshold)
 	{
diff --git a/Engine/Source/Runtime/SignalProcessing/Public/DSP/AudioFFT.h b/Engine/Source/Runtime/SignalProcessing/Public/DSP/AudioFFT.h
--- a/Engine/Source/Runtime/SignalProcessing/Public/DSP/AudioFFT.h
+++ b/Engine/Source/Runtime/SignalProcessing/Public/DSP/AudioFFT.h
@@ -77,7 +77,9 @@ namespace Audio
 			}
 
 			check(IsAligned<float*>(InBuffer, 4));
-			MultiplyBuffersInPlace(WindowBuffer.GetData(), InBuffer, NumSamples);
+			TArrayView<const float> WindowBufferView(WindowBuffer.GetData(), NumSamples);
+			TArrayView<float> InBufferView(InBuffer, NumSamples);
+			ArrayMultiplyInPlace(WindowBufferView, InBufferView);
 		}
 
 		EWindowType GetWindowType() const
diff --git a/Engine/Source/Runtime/SignalProcessing/Public/DSP/BufferVectorOperations.h b/Engine/Source/Runtime/SignalProcessing/Public/DSP/BufferVectorOperations.h
--- a/Engine/Source/Runtime/SignalProcessing/Public/DSP/BufferVectorOperations.h
+++ b/Engine/Source/Runtime/SignalProcessing/Public/DSP/BufferVectorOperations.h
@@ -28,109 +28,6 @@ namespace Audio
 	typedef TArray<float, FAudioBufferAlignedAllocator> AlignedFloatBuffer;
 	typedef TArray<int32, FAudioBufferAlignedAllocator> AlignedInt32Buffer;
 
-
-	/** CHANNEL-AGNOSTIC OPERATIONS */
-
-	/* Sets a values to zero if value is denormal. Denormal numbers significantly slow down floating point operations. */
-	SIGNALPROCESSING_API void BufferUnderflowClampFast(FAlignedFloatBuffer& InOutBuffer);
-
-	/* Sets a values to zero if value is denormal. Denormal numbers significantly slow down floating point operations. */
-	SIGNALPROCESSING_API void BufferUnderflowClampFast(float* RESTRICT InOutBuffer, const int32 InNum);
-	
-	/* Clamps the values in a buffer between a min and max value.*/
-	SIGNALPROCESSING_API void BufferRangeClampFast(FAlignedFloatBuffer& InOutBuffer, float InMinValue, float InMaxValue);
-	SIGNALPROCESSING_API void BufferRangeClampFast(float* RESTRICT InOutBuffer, const int32 InNum, float InMinValue, float InMaxValue);
-	
-	/** Multiplies the input aligned float buffer with the given value. */
-	SIGNALPROCESSING_API void BufferMultiplyByConstant(const FAlignedFloatBuffer& InFloatBuffer, float InValue, FAlignedFloatBuffer& OutFloatBuffer);
-	SIGNALPROCESSING_API void BufferMultiplyByConstant(const float* RESTRICT InFloatBuffer, float InValue, float* RESTRICT OutFloatBuffer, const int32 InNumSamples);
-	SIGNALPROCESSING_API void BufferMultiplyByConstant(const FAlignedFloatBuffer& InFloatBuffer, float InValue, FAlignedFloatBuffer& OutFloatBuffer);
-
-	/** Similar to BufferMultiplyByConstant, but (a) assumes a buffer length divisible by 4 and (b) performs the multiply in place. */
-	SIGNALPROCESSING_API void MultiplyBufferByConstantInPlace(FAlignedFloatBuffer& InBuffer, float InGain);
-	SIGNALPROCESSING_API void MultiplyBufferByConstantInPlace(float* RESTRICT InBuffer, int32 NumSamples, float InGain);
-
-	/** Adds a constant to a buffer (useful for DC offset removal) */
-	SIGNALPROCESSING_API void AddConstantToBufferInplace(FAlignedFloatBuffer& InBuffer, float Constant);
-	SIGNALPROCESSING_API void AddConstantToBufferInplace(float* RESTRICT InBuffer, int32 NumSamples, float Constant);
-
-	/** Sets a constant to a buffer (useful for DC offset application) */
-	SIGNALPROCESSING_API void BufferSetToConstantInplace(FAlignedFloatBuffer& InBuffer, float Constant);
-	SIGNALPROCESSING_API void BufferSetToConstantInplace(float* RESTRICT InBuffer, int32 NumSamples, float Constant);
-
-	/* Performs an element-wise weighted sum OutputBuffer = (InBuffer1 x InGain1) + (InBuffer2 x InGain2) */
-	SIGNALPROCESSING_API void BufferWeightedSumFast(const FAlignedFloatBuffer& InBuffer1, float InGain1, const FAlignedFloatBuffer& InBuffer2, float InGain2, FAlignedFloatBuffer& OutBuffer);
-
-	/* Performs an element-wise weighted sum OutputBuffer = (InBuffer1 x InGain1) + InBuffer2 */
-	SIGNALPROCESSING_API void BufferWeightedSumFast(const FAlignedFloatBuffer& InBuffer1, float InGain1, const FAlignedFloatBuffer& InBuffer2, FAlignedFloatBuffer& OutBuffer);
-
-	/* Performs an element-wise weighted sum OutputBuffer = (InBuffer1 x InGain1) + (InBuffer2 x InGain2) */
-	SIGNALPROCESSING_API void BufferWeightedSumFast(const float* RESTRICT InBuffer1, float InGain1, const float* RESTRICT InBuffer2, float InGain2, float* RESTRICT OutBuffer, int32 InNum);
-
-	/* Performs an element-wise weighted sum OutputBuffer = (InBuffer1 x InGain1) + InBuffer2 */
-	SIGNALPROCESSING_API void BufferWeightedSumFast(const float* RESTRICT InBuffer1, float InGain1, const float* RESTRICT InBuffer2, float* RESTRICT OutBuffer, int32 InNum);
-	SIGNALPROCESSING_API void MultiplyBufferByConstantInPlace(FAlignedFloatBuffer& InBuffer, float InGain);
-	SIGNALPROCESSING_API void MultiplyBufferByConstantInPlace(float* RESTRICT InBuffer, int32 NumSamples, float InGain);
-
-	/* Takes a float buffer and quickly interpolates it's gain from StartValue to EndValue. */
-	/* This operation completely ignores channel counts, so avoid using this function on buffers that are not mono, stereo or quad */
-	/* if the buffer needs to fade all channels uniformly. */
-	SIGNALPROCESSING_API void FadeBufferFast(FAlignedFloatBuffer& OutFloatBuffer, const float StartValue, const float EndValue);
-	SIGNALPROCESSING_API void FadeBufferFast(float* RESTRICT OutFloatBuffer, int32 NumSamples, const float StartValue, const float EndValue);
-
-	/** Takes buffer InFloatBuffer, optionally multiplies it by Gain, and adds it to BufferToSumTo. */
-	SIGNALPROCESSING_API void MixInBufferFast(const FAlignedFloatBuffer& InFloatBuffer, FAlignedFloatBuffer& BufferToSumTo, const float Gain);
-	SIGNALPROCESSING_API void MixInBufferFast(const float* RESTRICT InFloatBuffer, float* RESTRICT BufferToSumTo, int32 NumSamples, const float Gain);
-	SIGNALPROCESSING_API void MixInBufferFast(const FAlignedFloatBuffer& InFloatBuffer, FAlignedFloatBuffer& BufferToSumTo, const float StartGain, const float EndGain);
-	SIGNALPROCESSING_API void MixInBufferFast(const float* RESTRICT InFloatBuffer, float* RESTRICT BufferToSumTo, int32 NumSamples, const float StartGain, const float EndGain);
-	SIGNALPROCESSING_API void MixInBufferFast(const FAlignedFloatBuffer& InFloatBuffer, FAlignedFloatBuffer& BufferToSumTo);
-	SIGNALPROCESSING_API void MixInBufferFast(const float* RESTRICT InFloatBuffer, float* RESTRICT BufferToSumTo, int32 NumSamples);
-
-	/* Subtracts two buffers together element-wise. */
-	SIGNALPROCESSING_API void BufferSubtractFast(const FAlignedFloatBuffer& InMinuend, const FAlignedFloatBuffer& InSubtrahend, FAlignedFloatBuffer& OutputBuffer);
-
-	/* Subtracts two buffers together element-wise. */
-	SIGNALPROCESSING_API void BufferSubtractFast(const float* RESTRICT InMinuend, const float* RESTRICT InSubtrahend, float* RESTRICT OutputBuffer, int32 NumSamples);
-
-	/* Performs element-wise in-place subtraction placing the result in the subtrahend. InOutSubtrahend = InMinuend - InOutSubtrahend */
-	SIGNALPROCESSING_API void BufferSubtractInPlace1Fast(const FAlignedFloatBuffer& InMinuend, FAlignedFloatBuffer& InOutSubtrahend);
-
-	/* Performs element-wise in-place subtraction placing the result in the subtrahend. InOutSubtrahend = InMinuend - InOutSubtrahend */
-	SIGNALPROCESSING_API void BufferSubtractInPlace1Fast(const float* RESTRICT InMinuend, float* RESTRICT InOutSubtrahend, int32 NumSamples);
-
-
-	/* Performs element-wise in-place subtraction placing the result in the minuend. InOutMinuend = InOutMinuend - InSubtrahend */
-	SIGNALPROCESSING_API void BufferSubtractInPlace2Fast(FAlignedFloatBuffer& InOutMinuend, const FAlignedFloatBuffer& InSubtrahend);
-
-	/* Performs element-wise in-place subtraction placing the result in the minuend. InOutMinuend = InOutMinuend - InSubtrahend */
-	SIGNALPROCESSING_API void BufferSubtractInPlace2Fast(float* RESTRICT InOutMinuend, const float* RESTRICT InSubtrahend, int32 NumSamples);
-
-	/** This version of MixInBufferFast will fade from StartGain to EndGain. */
-	SIGNALPROCESSING_API void MixInBufferFast(const FAlignedFloatBuffer& InFloatBuffer, FAlignedFloatBuffer& BufferToSumTo, const float StartGain, const float EndGain);
-	SIGNALPROCESSING_API void MixInBufferFast(const float* RESTRICT InFloatBuffer, float* RESTRICT BufferToSumTo, int32 NumSamples, const float StartGain, const float EndGain);
-
-	/** Sums two buffers together and places the result in the resulting buffer. */
-	SIGNALPROCESSING_API void SumBuffers(const FAlignedFloatBuffer& InFloatBuffer1, const FAlignedFloatBuffer& InFloatBuffer2, FAlignedFloatBuffer& OutputBuffer);
-	SIGNALPROCESSING_API void SumBuffers(const float* RESTRICT InFloatBuffer1, const float* RESTRICT InFloatBuffer2, float* RESTRICT OutputBuffer, int32 NumSamples);
-
-	/** Multiply the second buffer in place by the first buffer. */
-	SIGNALPROCESSING_API void MultiplyBuffersInPlace(const FAlignedFloatBuffer& InFloatBuffer, FAlignedFloatBuffer& BufferToMultiply);
-	SIGNALPROCESSING_API void MultiplyBuffersInPlace(const float* RESTRICT InFloatBuffer, float* RESTRICT BufferToMultiply, int32 NumSamples);
-
-	/** CHANNEL-AGNOSTIC ANALYSIS OPERATIONS */
-
-	/** Takes an audio buffer and returns the magnitude across that buffer. */
-	SIGNALPROCESSING_API float GetMagnitude(const FAlignedFloatBuffer& Buffer);
-	SIGNALPROCESSING_API float GetMagnitude(const float* RESTRICT Buffer, int32 NumSamples);
-
-	/** Takes an audio buffer and gets the average amplitude across that buffer. */
-	SIGNALPROCESSING_API float BufferGetAverageValue(const FAlignedFloatBuffer& Buffer);
-	SIGNALPROCESSING_API float BufferGetAverageValue(const float* RESTRICT Buffer, int32 NumSamples);
-
-	/** Takes an audio buffer and gets the average absolute amplitude across that buffer. */
-	SIGNALPROCESSING_API float BufferGetAverageAbsValue(const FAlignedFloatBuffer& Buffer);
-	SIGNALPROCESSING_API float BufferGetAverageAbsValue(const float* RESTRICT Buffer, int32 NumSamples);
-
 	/** CHANNEL-SPECIFIC OPERATIONS */
 
 	/** Takes a 2 channel interleaved buffer and applies Gains to it. Gains is expected to point to a 2 float long buffer.
@@ -274,18 +171,6 @@ namespace Audio
 	/** Sums 2 channel interleaved input samples. OutSamples[n] = InSamples[2n] + InSamples[2n + 1] */
 	SIGNALPROCESSING_API void BufferSum2ChannelToMonoFast(const float* RESTRICT InSamples, float* RESTRICT OutSamples, const int32 InNumFrames);
 
-	/** Compute power of complex data. Out[i] = Real[i] * Real[i] + Imaginary[i] * Imaginary[i] */
-	SIGNALPROCESSING_API void BufferComplexToPowerFast(const FAlignedFloatBuffer& InRealSamples, const FAlignedFloatBuffer& InImaginarySamples, FAlignedFloatBuffer& OutPowerSamples);
-
-	/** Compute power of complex data. Out[i] = Real[i] * Real[i] + Imaginary[i] * Imaginary[i] */
-	SIGNALPROCESSING_API void BufferComplexToPowerFast(const float* RESTRICT InRealSamples, const float* RESTRICT InImaginarySamples, float* RESTRICT OutPowerSamples, const int32 InNum);
-
-	/** Compute magnitude of complex data. Out[i] = Sqrt(Real[i] * Real[i] + Imaginary[i] * Imaginary[i]) */
-	SIGNALPROCESSING_API void BufferComplexToMagnitudeFast(const FAlignedFloatBuffer& InRealSamples, const FAlignedFloatBuffer& InImaginarySamples, FAlignedFloatBuffer& OutPowerSamples);
-
-	/** Compute magnitude of complex data. Out[i] = Sqrt(Real[i] * Real[i] + Imaginary[i] * Imaginary[i]) */
-	SIGNALPROCESSING_API void BufferComplexToMagnitudeFast(const float* RESTRICT InRealSamples, const float* RESTRICT InImaginarySamples, float* RESTRICT OutPowerSamples, const int32 InNum);
-
 	/** Class which handles a vectorized interpolation of an entire buffer to the values of a target buffer */
 	class SIGNALPROCESSING_API FBufferLinearEase
 	{
diff --git a/Engine/Source/Runtime/SignalProcessing/Public/DSP/FloatArrayMath.h b/Engine/Source/Runtime/SignalProcessing/Public/DSP/FloatArrayMath.h
--- a/Engine/Source/Runtime/SignalProcessing/Public/DSP/FloatArrayMath.h
+++ b/Engine/Source/Runtime/SignalProcessing/Public/DSP/FloatArrayMath.h
@@ -11,6 +11,10 @@ namespace Audio
 	SIGNALPROCESSING_API void ArraySum(TArrayView<const float> InValues, float& OutSum);
 	SIGNALPROCESSING_API void ArraySum(const FAlignedFloatBuffer& InValues, float& OutSum);
 
+	/** Sums two buffers together and places the result in the resulting buffer. */
+	SIGNALPROCESSING_API void ArraySum(TArrayView<const float> InFloatBuffer1, TArrayView<const float> InFloatBuffer2, TArrayView<float> OutputBuffer);
+	SIGNALPROCESSING_API void ArraySum(const FAlignedFloatBuffer& InFloatBuffer1, const FAlignedFloatBuffer& InFloatBuffer2, FAlignedFloatBuffer& OutputBuffer);
+
 	/** Cumulative sum of array.
 	 *
 	 *  InView contains data to be cumulatively summed.
@@ -32,6 +36,18 @@ namespace Audio
 	 */
 	SIGNALPROCESSING_API void ArrayMeanSquared(TArrayView<const float> InView, float& OutMean);
 
+	/** Takes an audio buffer and returns the magnitude across that buffer. */
+	SIGNALPROCESSING_API float ArrayGetMagnitude(TArrayView<const float> Buffer);
+	SIGNALPROCESSING_API float ArrayGetMagnitude(const FAlignedFloatBuffer& Buffer);
+
+	/** Takes an audio buffer and gets the average amplitude across that buffer. */
+	SIGNALPROCESSING_API float ArrayGetAverageValue(TArrayView<const float> Buffer);
+	SIGNALPROCESSING_API float ArrayGetAverageValue(const FAlignedFloatBuffer& Buffer);
+
+	/** Takes an audio buffer and gets the average absolute amplitude across that buffer. */
+	SIGNALPROCESSING_API float ArrayGetAverageAbsValue(TArrayView<const float> Buffer);
+	SIGNALPROCESSING_API float ArrayGetAverageAbsValue(const FAlignedFloatBuffer& Buffer);
+
 	/** Mean filter of array.
 	 *
 	 *  Note: Uses standard biased mean estimator of Sum(x) / Count(x).
@@ -96,9 +112,9 @@ namespace Audio
 	 */ 
 	SIGNALPROCESSING_API void ArrayMinMaxNormalize(TArrayView<const float> InView, TArray<float>& OutArray);
 
-	/** Multiplies two arrays element-wise. Stores result in InValues2*/
-	SIGNALPROCESSING_API void ArrayMultiplyInPlace(TArrayView<const float> InValues1, TArrayView<float> InValues2);
-	SIGNALPROCESSING_API void ArrayMultiplyInPlace(const FAlignedFloatBuffer& InValues1, FAlignedFloatBuffer& InValues2);
+	/** Multiply the second buffer in place by the first buffer. */
+	SIGNALPROCESSING_API void ArrayMultiplyInPlace(TArrayView<const float> InFloatBuffer, TArrayView<float> BufferToMultiply);
+	SIGNALPROCESSING_API void ArrayMultiplyInPlace(const FAlignedFloatBuffer& InFloatBuffer, FAlignedFloatBuffer& BufferToMultiply);
 
 	/** Multiplies two complex valued arrays element-wise. 
 	 * This assumes elements are in interleaved format [real_0, imag_0, ..., real_N, imag_N]
@@ -107,9 +123,13 @@ namespace Audio
 	SIGNALPROCESSING_API void ArrayComplexMultiplyInPlace(TArrayView<const float> InValues1, TArrayView<float> InValues2);
 	SIGNALPROCESSING_API void ArrayComplexMultiplyInPlace(const FAlignedFloatBuffer& InValues1, FAlignedFloatBuffer& InValues2);
 
-	/** Multiplies each element in InView by InMultiplier */
-	SIGNALPROCESSING_API void ArrayMultiplyByConstantInPlace(TArrayView<float> InValues, float InMultiplier);
-	SIGNALPROCESSING_API void ArrayMultiplyByConstantInPlace(FAlignedFloatBuffer& InValues, float InMultiplier);
+	/** Multiplies the input float buffer with the given value. */
+	SIGNALPROCESSING_API void ArrayMultiplyByConstant(TArrayView<const float> InFloatBuffer, float InValue, TArrayView<float> OutFloatBuffer);
+	SIGNALPROCESSING_API void ArrayMultiplyByConstant(const FAlignedFloatBuffer& InFloatBuffer, float InValue, FAlignedFloatBuffer& OutFloatBuffer);
+
+	/** Similar to ArrayMultiplyByConstant, but performs the multiply in place. */
+	SIGNALPROCESSING_API void ArrayMultiplyByConstantInPlace(TArrayView<float> InOutBuffer, float InGain);
+	SIGNALPROCESSING_API void ArrayMultiplyByConstantInPlace(FAlignedFloatBuffer& InOutBuffer, float InGain);
 
 	/** Add arrays element-wise in place. InAccumulateValues[i] += InValues[i]
 	 *
@@ -119,6 +139,10 @@ namespace Audio
 	SIGNALPROCESSING_API void ArrayAddInPlace(TArrayView<const float> InValues, TArrayView<float> InAccumulateValues);
 	SIGNALPROCESSING_API void ArrayAddInPlace(const FAlignedFloatBuffer& InValues, FAlignedFloatBuffer& InAccumulateValues);
 
+	/** Adds a constant to a buffer (useful for DC offset removal) */
+	SIGNALPROCESSING_API void ArrayAddConstantInplace(TArrayView<float> InOutBuffer, float InConstant);
+	SIGNALPROCESSING_API void ArrayAddConstantInplace(FAlignedFloatBuffer& InOutBuffer, float InConstant);
+
 	/** Multiply Add arrays element-wise in place. InAccumulateValues[i] += InMultiplier * InValues[i]
 	 *
 	 *  @param InValues - The array to add.
@@ -143,10 +167,18 @@ namespace Audio
 	 *
 	 *  InMinuend is the array of data to be subtracted from.
 	 *  InSubtrahend is the array of data to subtract.
-	 *  OutArray is the array which holds the result.
+	 *  OutBuffer is the array which holds the result.
 	 */
-	SIGNALPROCESSING_API void ArraySubtract(TArrayView<const float> InMinuend, TArrayView<const float> InSubtrahend, TArray<float>& OutArray);
-	SIGNALPROCESSING_API void ArraySubtract(const FAlignedFloatBuffer& InMinuend, const FAlignedFloatBuffer& InSubtrahend, FAlignedFloatBuffer& OutArray);
+	SIGNALPROCESSING_API void ArraySubtract(TArrayView<const float> InMinuend, TArrayView<const float> InSubtrahend, TArrayView<float> OutBuffer);
+	SIGNALPROCESSING_API void ArraySubtract(const FAlignedFloatBuffer& InMinuend, const FAlignedFloatBuffer& InSubtrahend, FAlignedFloatBuffer& OutputBuffer);
+
+	/* Performs element-wise in-place subtraction placing the result in the subtrahend. InOutSubtrahend = InMinuend - InOutSubtrahend */
+	SIGNALPROCESSING_API void ArraySubtractInPlace1(TArrayView<const float> InMinuend, TArrayView<float> InOutSubtrahend);
+	SIGNALPROCESSING_API void ArraySubtractInPlace1(const FAlignedFloatBuffer& InMinuend, FAlignedFloatBuffer& InOutSubtrahend);
+
+	/* Performs element-wise in-place subtraction placing the result in the minuend. InOutMinuend = InOutMinuend - InSubtrahend */
+	SIGNALPROCESSING_API void ArraySubtractInPlace2(TArrayView<float> InOutMinuend, TArrayView<const float> InSubtrahend);
+	SIGNALPROCESSING_API void ArraySubtractInPlace2(FAlignedFloatBuffer& InOutMinuend, const FAlignedFloatBuffer& InSubtrahend);
 
 	/** Subtract value from each element in InValues */
 	SIGNALPROCESSING_API void ArraySubtractByConstantInPlace(TArrayView<float> InValues, float InSubtrahend);
@@ -161,7 +193,6 @@ namespace Audio
 	/** Take Square Root of values in place. */
 	SIGNALPROCESSING_API void ArraySqrtInPlace(TArrayView<float> InValues);
 
-
 	/** Perform complex conjugate of array.  Assumes complex numbers are interlaves [real_0, imag_0, real_1, image_1, ..., real_N, imag_N]. */
 	SIGNALPROCESSING_API void ArrayComplexConjugate(TArrayView<const float> InValues, TArrayView<float> OutValues);
 	SIGNALPROCESSING_API void ArrayComplexConjugate(const FAlignedFloatBuffer& InValues, FAlignedFloatBuffer& OutValues);
@@ -177,14 +208,49 @@ namespace Audio
 	SIGNALPROCESSING_API void ArrayPowerToDecibelInPlace(TArrayView<float> InValues, float InMinimumDb);
 	SIGNALPROCESSING_API void ArrayPowerToDecibelInPlace(FAlignedFloatBuffer& InValues, float InMinimumDb);
 
-
-	/** Compute power of complex data. Out[i] = Complex[2 * i] * Complex[2 * i] + Complex[2 * i + 1] * Complex[2 * i + 1] */
-	SIGNALPROCESSING_API void ArrayComplexToPower(TArrayView<float> InComplexSamples, TArray<float>& OutPowerSamples);
-
 	/** Compute power of complex data. Out[i] = Complex[2 * i] * Complex[2 * i] + Complex[2 * i + 1] * Complex[2 * i + 1] */
+	SIGNALPROCESSING_API void ArrayComplexToPower(TArrayView<const float> InComplexSamples, TArrayView<float> OutPowerSamples);
 	SIGNALPROCESSING_API void ArrayComplexToPower(const FAlignedFloatBuffer& InComplexSamples, FAlignedFloatBuffer& OutPowerSamples);
 
+	/** Compute power of complex data. Out[i] = Real[i] * Real[i] + Imaginary[i] * Imaginary[i] */
+	SIGNALPROCESSING_API void ArrayComplexToPower(TArrayView<const float> InRealSamples, TArrayView<const float> InImaginarySamples, TArrayView<float> OutPowerSamples);
+	SIGNALPROCESSING_API void ArrayComplexToPower(const FAlignedFloatBuffer& InRealSamples, const FAlignedFloatBuffer& InImaginarySamples, FAlignedFloatBuffer& OutPowerSamples);
+
+	/* Sets a values to zero if value is denormal. Denormal numbers significantly slow down floating point operations. */
+	SIGNALPROCESSING_API void ArrayUnderflowClamp(TArrayView<float> InOutBuffer);
+	SIGNALPROCESSING_API void ArrayUnderflowClamp(FAlignedFloatBuffer& InOutBuffer);
+
+	/* Clamps the values in a buffer between a min and max value.*/
+	SIGNALPROCESSING_API void ArrayRangeClamp(TArrayView<float> InOutBuffer, float InMinValue, float InMaxValue);
+	SIGNALPROCESSING_API void ArrayRangeClamp(FAlignedFloatBuffer& InOutBuffer, float InMinValue, float InMaxValue);
+
+	/** Sets a constant to a buffer (useful for DC offset application) */
+	SIGNALPROCESSING_API void ArraySetToConstantInplace(TArrayView<float> InOutBuffer, float InConstant);
+	SIGNALPROCESSING_API void ArraySetToConstantInplace(FAlignedFloatBuffer& InOutBuffer, float InConstant);
 
+	/* Performs an element-wise weighted sum OutputBuffer = (InBuffer1 x InGain1) + (InBuffer2 x InGain2) */
+	SIGNALPROCESSING_API void ArrayWeightedSum(TArrayView<const float> InBuffer1, float InGain1, TArrayView<const float> InBuffer2, float InGain2, TArrayView<float> OutBuffer);
+	SIGNALPROCESSING_API void ArrayWeightedSum(const FAlignedFloatBuffer& InBuffer1, float InGain1, const FAlignedFloatBuffer& InBuffer2, float InGain2, FAlignedFloatBuffer& OutBuffer);
+
+	/* Performs an element-wise weighted sum OutputBuffer = (InBuffer1 x InGain1) + InBuffer2 */
+	SIGNALPROCESSING_API void ArrayWeightedSum(TArrayView<const float> InBuffer1, float InGain1, TArrayView<const float> InBuffer2, TArrayView<float> OutBuffer);
+	SIGNALPROCESSING_API void ArrayWeightedSum(const FAlignedFloatBuffer& InBuffer1, float InGain1, const FAlignedFloatBuffer& InBuffer2, FAlignedFloatBuffer& OutBuffer);
+
+	/* Takes a float buffer and quickly interpolates it's gain from StartValue to EndValue. */
+	/* This operation completely ignores channel counts, so avoid using this function on buffers that are not mono, stereo or quad */
+	/* if the buffer needs to fade all channels uniformly. */
+	SIGNALPROCESSING_API void ArrayFade(TArrayView<float> InOutBuffer, const float StartValue, const float EndValue);
+	SIGNALPROCESSING_API void ArrayFade(FAlignedFloatBuffer& OutFloatBuffer, const float StartValue, const float EndValue);
+
+	/** Takes buffer InFloatBuffer, optionally multiplies it by Gain, and adds it to BufferToSumTo. */
+	SIGNALPROCESSING_API void ArrayMixIn(TArrayView<const float> InFloatBuffer, TArrayView<float> BufferToSumTo, const float Gain);
+	SIGNALPROCESSING_API void ArrayMixIn(const FAlignedFloatBuffer& InFloatBuffer, FAlignedFloatBuffer& BufferToSumTo, const float Gain);
+	SIGNALPROCESSING_API void ArrayMixIn(TArrayView<const float> InFloatBuffer, TArrayView<float> BufferToSumTo);
+	SIGNALPROCESSING_API void ArrayMixIn(const FAlignedFloatBuffer& InFloatBuffer, FAlignedFloatBuffer& BufferToSumTo);
+
+	/** This version of ArrayMixIn will fade from StartGain to EndGain. */
+	SIGNALPROCESSING_API void ArrayMixIn(TArrayView<const float> InFloatBuffer, TArrayView<float> BufferToSumTo, const float StartGain, const float EndGain);
+	SIGNALPROCESSING_API void ArrayMixIn(const FAlignedFloatBuffer& InFloatBuffer, FAlignedFloatBuffer& BufferToSumTo, const float StartGain, const float EndGain);
 
 	/** FContiguousSparse2DKernelTransform
 	 *
@@ -260,3 +326,12 @@ namespace Audio
 	};
 
 }
+
+// Support ISPC enable/disable in non-shipping builds
+#if !INTEL_ISPC
+const bool bAudio_FloatArrayMath_ISPC_Enabled = false;
+#elif UE_BUILD_SHIPPING
+const bool bAudio_FloatArrayMath_ISPC_Enabled = true;
+#else
+extern bool bAudio_FloatArrayMath_ISPC_Enabled;
+#endif
diff --git a/Engine/Source/Runtime/SignalProcessing/SignalProcessing.Build.cs b/Engine/Source/Runtime/SignalProcessing/SignalProcessing.Build.cs
--- a/Engine/Source/Runtime/SignalProcessing/SignalProcessing.Build.cs
+++ b/Engine/Source/Runtime/SignalProcessing/SignalProcessing.Build.cs
@@ -13,6 +13,13 @@ namespace UnrealBuildTool.Rules
 				}
 			);
 
+			PublicDependencyModuleNames.AddRange(
+				new string[] 
+				{
+					"IntelISPC"
+				}
+			);
+
 			AddEngineThirdPartyPrivateStaticDependencies(Target,
 					"UEOgg",
 					"Vorbis",
