PR #8970: Fix FFT: Incorrect calculation of real values of spectrum (Contributed by [REDACTED_PERSON])

diff --git a/Engine/Source/Runtime/SignalProcessing/Private/AudioFFT.cpp b/Engine/Source/Runtime/SignalProcessing/Private/AudioFFT.cpp
--- a/Engine/Source/Runtime/SignalProcessing/Private/AudioFFT.cpp
+++ b/Engine/Source/Runtime/SignalProcessing/Private/AudioFFT.cpp
@@ -435,7 +435,7 @@ namespace Audio
 		void PerformDFT(const FFTTimeDomainData& InputParams, FFTFreqDomainData& OutputParams)
 		{
 			const float* InputBuffer = InputParams.Buffer;
-			float* OutReal = OutputParams.OutImag;
+			float* OutReal = OutputParams.OutReal;
 			float* OutImag = OutputParams.OutImag;
 
 			float N = InputParams.NumSamples;
@@ -460,7 +460,7 @@ namespace Audio
 		void PerformIDFT(const FFTFreqDomainData& InputParams, FFTTimeDomainData& OutputParams)
 		{
 			float* OutputBuffer = OutputParams.Buffer;
-			float* InReal = InputParams.OutImag;
+			float* InReal = InputParams.OutReal;
 			float* InImag = InputParams.OutImag;
 
 			float N = OutputParams.NumSamples;
