Add Clip Normalization and float/int16 conversion audio array functions

diff --git a/Engine/Source/Runtime/SignalProcessing/Private/FloatArrayMath.cpp b/Engine/Source/Runtime/SignalProcessing/Private/FloatArrayMath.cpp
--- a/Engine/Source/Runtime/SignalProcessing/Private/FloatArrayMath.cpp
+++ b/Engine/Source/Runtime/SignalProcessing/Private/FloatArrayMath.cpp
@@ -2,6 +2,7 @@
 
 #include "DSP/FloatArrayMath.h"
 #include "CoreMinimal.h"
+#include "Math/UnrealMathVectorConstants.h"
 #include "SignalProcessingModule.h"
 
 #if INTEL_ISPC && !UE_BUILD_SHIPPING
@@ -707,6 +708,42 @@ namespace Audio
 			}
 		}
 	}
+	
+	float ArrayMaxAbsValue(const TArrayView<float> InView)
+	{
+		const int32 Num = InView.Num();
+		const float* Data = InView.GetData();
+
+		float Max = 0.f;
+		const int32 NumToSimd = Num & MathIntrinsics::SimdMask;
+		const int32 NumNotToSimd = Num & MathIntrinsics::NotSimdMask;
+
+		VectorRegister4Float MaxVector = VectorSetFloat1(0.f);
+		
+		if (NumToSimd)
+		{
+			for (int32 i = 0; i < NumToSimd; i += AUDIO_NUM_FLOATS_PER_VECTOR_REGISTER)
+			{
+				VectorRegister4Float Input1 = VectorLoad(&Data[i]);
+
+				MaxVector = VectorMax(MaxVector, VectorAbs(Input1));
+			}
+
+			AlignedFloat4 OutArray(MaxVector);
+			
+			Max =  FMath::Max(FMath::Max(OutArray[0], OutArray[1]), FMath::Max(OutArray[2], OutArray[3]));
+		}
+
+		if (NumNotToSimd)
+		{
+			for (int32 i = NumToSimd; i < Num; ++i)
+			{
+				Max = FMath::Max(FMath::Abs(Data[i]), Max);
+			}
+		}
+
+		return Max;
+	}
 
 	void ArrayMultiplyInPlace(TArrayView<const float> InFloatBuffer, TArrayView<float> BufferToMultiply)
 	{
@@ -2262,6 +2299,90 @@ namespace Audio
 		ArrayMixIn(InFloatBufferView, BufferToSumToView, StartGain, EndGain);
 	}
 
+	void ArrayFloatToPcm16(TArrayView<const float> InView, TArrayView<int16> OutView)
+	{
+		check(InView.Num() == OutView.Num());
+
+		const int32 Num = InView.Num();
+		
+		const int32 NumToSimd = Num & MathIntrinsics::SimdMask;
+		const int32 NumNotToSimd = Num & MathIntrinsics::NotSimdMask;
+
+		const float* InputPtr = InView.GetData();
+		int16* OutPtr = OutView.GetData();
+
+		constexpr float ConversionValue = static_cast<float>(TNumericLimits<int16>::Max());
+
+		if(NumToSimd)
+		{
+			const VectorRegister4Float ConversionVector = VectorSetFloat1(ConversionValue);
+
+			for (int32 i = 0; i < NumToSimd; i += AUDIO_NUM_FLOATS_PER_VECTOR_REGISTER)
+			{
+				const VectorRegister4Float InVector = VectorLoad(&InputPtr[i]);
+				const VectorRegister4Float ScaledVector = VectorMultiply(InVector, ConversionVector);
+				const VectorRegister4Int IntVector = VectorFloatToInt(ScaledVector);
+
+				const AlignedFloat4 ScaledFloatArray(ScaledVector);
+
+				OutPtr[i] =		(int16)ScaledFloatArray[0];
+				OutPtr[i + 1] =	(int16)ScaledFloatArray[1];
+				OutPtr[i + 2] =	(int16)ScaledFloatArray[2];
+				OutPtr[i + 3] =	(int16)ScaledFloatArray[3];
+			}
+		}
+
+		if(NumNotToSimd)
+		{
+			for (int32 i = NumToSimd; i < Num; i++)
+			{
+				OutPtr[i] = (int16)(InputPtr[i] * ConversionValue);
+			}
+		}
+	}
+	
+	void ArrayPcm16ToFloat(TArrayView<const int16> InView, TArrayView<float> OutView)
+	{
+		check(InView.Num() == OutView.Num());
+
+		const int32 Num = InView.Num();
+		
+		const int32 NumToSimd = Num & MathIntrinsics::SimdMask;
+		const int32 NumNotToSimd = Num & MathIntrinsics::NotSimdMask;
+
+		const int16* InputPtr = InView.GetData();
+		float* OutPtr = OutView.GetData();
+
+		constexpr float ConversionValue = 1.f / static_cast<float>(TNumericLimits<int16>::Max());
+
+		if(NumToSimd)
+		{
+			const VectorRegister4Float ConversionVector = VectorSetFloat1(ConversionValue);
+			AlignedFloat4 FloatArray(GlobalVectorConstants::FloatZero);
+
+			for (int32 i = 0; i < NumToSimd; i += AUDIO_NUM_FLOATS_PER_VECTOR_REGISTER)
+			{
+				FloatArray[0] = (float)InputPtr[i];
+				FloatArray[1] = (float)InputPtr[i + 1];
+				FloatArray[2] = (float)InputPtr[i + 2];
+				FloatArray[3] = (float)InputPtr[i + 3];
+
+				const VectorRegister4Float InVector = FloatArray.ToVectorRegister();
+				const VectorRegister4Float ScaledVector = VectorMultiply(InVector, ConversionVector);
+			
+				VectorStore(ScaledVector, &OutPtr[i]);
+			}
+		}
+
+		if(NumNotToSimd)
+		{
+			for (int32 i = NumToSimd; i < Num; i++)
+			{
+				OutPtr[i] = (float)InputPtr[i] * ConversionValue;
+			}
+		}
+	}
+
 	FContiguousSparse2DKernelTransform::FContiguousSparse2DKernelTransform(const int32 NumInElements, const int32 NumOutElements)
 	:	NumIn(NumInElements)
 	,	NumOut(NumOutElements)
diff --git a/Engine/Source/Runtime/SignalProcessing/Public/DSP/FloatArrayMath.h b/Engine/Source/Runtime/SignalProcessing/Public/DSP/FloatArrayMath.h
--- a/Engine/Source/Runtime/SignalProcessing/Public/DSP/FloatArrayMath.h
+++ b/Engine/Source/Runtime/SignalProcessing/Public/DSP/FloatArrayMath.h
@@ -112,6 +112,9 @@ namespace Audio
 	 */ 
 	SIGNALPROCESSING_API void ArrayMinMaxNormalize(TArrayView<const float> InView, TArray<float>& OutArray);
 
+	// Scale an array to fit within [-1.f, 1.f]
+	SIGNALPROCESSING_API float ArrayMaxAbsValue(const TArrayView<float> InView);
+
 	/** Multiply the second buffer in place by the first buffer. */
 	SIGNALPROCESSING_API void ArrayMultiplyInPlace(TArrayView<const float> InFloatBuffer, TArrayView<float> BufferToMultiply);
 	SIGNALPROCESSING_API void ArrayMultiplyInPlace(const FAlignedFloatBuffer& InFloatBuffer, FAlignedFloatBuffer& BufferToMultiply);
@@ -252,6 +255,9 @@ namespace Audio
 	SIGNALPROCESSING_API void ArrayMixIn(TArrayView<const float> InFloatBuffer, TArrayView<float> BufferToSumTo, const float StartGain, const float EndGain);
 	SIGNALPROCESSING_API void ArrayMixIn(const FAlignedFloatBuffer& InFloatBuffer, FAlignedFloatBuffer& BufferToSumTo, const float StartGain, const float EndGain);
 
+	SIGNALPROCESSING_API void ArrayFloatToPcm16(TArrayView<const float> InView, TArrayView<int16> OutView);
+	SIGNALPROCESSING_API void ArrayPcm16ToFloat(TArrayView<const int16> InView, TArrayView<float> OutView);
+
 	/** FContiguousSparse2DKernelTransform
 	 *
 	 *  FContiguousSparse2DKernelTransform applies a matrix transformation to an input array. 
