Fixing crash in tone generator

diff --git a/Engine/Plugins/Runtime/Synthesis/Source/Synthesis/Private/SynthComponentToneGenerator.cpp b/Engine/Plugins/Runtime/Synthesis/Source/Synthesis/Private/SynthComponentToneGenerator.cpp
--- a/Engine/Plugins/Runtime/Synthesis/Source/Synthesis/Private/SynthComponentToneGenerator.cpp
+++ b/Engine/Plugins/Runtime/Synthesis/Source/Synthesis/Private/SynthComponentToneGenerator.cpp
@@ -6,6 +6,9 @@
 
 #define TEST_INT16_ATTENUATION 1
 
+// Disabling this experimental feature for this synth component 
+#define SYNTH_COMP_ENABLE_DIST_ATTENUATION 0
+
 FToneGenerator::FToneGenerator(int32 InSampleRate, int32 InNumChannels, int32 InFrequency, float InVolume, const Audio::FAudioBufferDistanceAttenuationSettings& InAttenuationSettings)
 	: NumChannels(InNumChannels)
 {
@@ -54,6 +57,7 @@ int32 FToneGenerator::OnGenerateAudio(float* OutAudio, int32 NumSamples)
 		}
 	}
 
+#if SYNTH_COMP_ENABLE_DIST_ATTENUATION
 #if TEST_INT16_ATTENUATION
 
 	// Convert to int16
@@ -77,6 +81,8 @@ int32 FToneGenerator::OnGenerateAudio(float* OutAudio, int32 NumSamples)
 
 	Audio::DistanceAttenuationProcessAudio(AudioBufferView, NumChannels, CurrentDistance, DistanceAttenuationSettings, CurrentAttenuation);
 #endif
+#endif
+
 	return NumSamples;
 }
 
diff --git a/Engine/Source/Runtime/SignalProcessing/Private/AudioBufferDistanceAttenuation.cpp b/Engine/Source/Runtime/SignalProcessing/Private/AudioBufferDistanceAttenuation.cpp
--- a/Engine/Source/Runtime/SignalProcessing/Private/AudioBufferDistanceAttenuation.cpp
+++ b/Engine/Source/Runtime/SignalProcessing/Private/AudioBufferDistanceAttenuation.cpp
@@ -14,10 +14,18 @@ namespace Audio
 
 		float CurveValue = 0.0f;
 
-		bool bSuccess = InSettings.AttenuationCurve.Eval(Alpha, CurveValue);
+		if (InSettings.AttenuationCurve.Num() > 0)
+		{
+			bool bSuccess = InSettings.AttenuationCurve.Eval(Alpha, CurveValue);
+			// This should succeed since we ensure we always have at least two points in the curve when it's set
+			check(bSuccess);
+		}
+		else
+		{
+			// We do a linear attenuation if there is no curve
+			CurveValue = 1.0f - Alpha;
+		}
 
-		// This should succeed since we ensure we always have at least two points in the curve when it's set
-		check(bSuccess);
 
 		// Note the curve is expected to map to the attenuation amount (i.e. at right-most value, it'll be 0.0, which corresponds to max dB attenuation)
 		// This then needs to be used to interpolate the dB range (0.0 is no attenuation, -60 for example, is a lot of attenuation)
