Static analysis fix in ChannelMap.cpp

#rnx

diff --git a/Engine/Source/Runtime/SignalProcessing/Private/ChannelMap.cpp b/Engine/Source/Runtime/SignalProcessing/Private/ChannelMap.cpp
--- a/Engine/Source/Runtime/SignalProcessing/Private/ChannelMap.cpp
+++ b/Engine/Source/Runtime/SignalProcessing/Private/ChannelMap.cpp
@@ -435,8 +435,7 @@ namespace Audio
 	{
 		using namespace ChannelMapPrivate;
 
-		const bool bIsValidNumOutputChannels = (InNumOutputChannels > 0) && (InNumOutputChannels <= ChannelMapMaxNumChannels);
-		if (!ensureMsgf(bIsValidNumOutputChannels, TEXT("Invalid number of output channels: %d"), InNumOutputChannels))
+		if (!ensureMsgf((InNumOutputChannels > 0) && (InNumOutputChannels <= ChannelMapMaxNumChannels), TEXT("Invalid number of output channels: %d"), InNumOutputChannels))
 		{
 			return false;
 		}
