Converting audio routines to use single precision SIMD operations.

diff --git a/Engine/Source/Runtime/SignalProcessing/Private/LongDelayAPF.cpp b/Engine/Source/Runtime/SignalProcessing/Private/LongDelayAPF.cpp
--- a/Engine/Source/Runtime/SignalProcessing/Private/LongDelayAPF.cpp
+++ b/Engine/Source/Runtime/SignalProcessing/Private/LongDelayAPF.cpp
@@ -131,25 +131,25 @@ void FLongDelayAPF::ProcessAudioBlock(const float* InSamples, const float* InDel
 	// Calculate new delay line samples. "w[n] = x[n] + gw[n - d]"
 	int32 NumToSIMD = InNum - (InNum % AUDIO_NUM_FLOATS_PER_VECTOR_REGISTER);
 
-	VectorRegister VG = MakeVectorRegister(G, G, G, G);
-	VectorRegister VNG = MakeVectorRegister(-G, -G, -G, -G);
-	VectorRegister VFMIN = MakeVectorRegister(FLT_MIN, FLT_MIN, FLT_MIN, FLT_MIN);
-	VectorRegister VNFMIN = MakeVectorRegister(-FLT_MIN, -FLT_MIN, -FLT_MIN, -FLT_MIN);
+	VectorRegister4f VG = MakeVectorRegisterFloat(G, G, G, G);
+	VectorRegister4f VNG = MakeVectorRegisterFloat(-G, -G, -G, -G);
+	VectorRegister4f VFMIN = MakeVectorRegisterFloat(FLT_MIN, FLT_MIN, FLT_MIN, FLT_MIN);
+	VectorRegister4f VNFMIN = MakeVectorRegisterFloat(-FLT_MIN, -FLT_MIN, -FLT_MIN, -FLT_MIN);
 
 	for (int32 i = 0; i < InNum; i += 4)
 	{
-		VectorRegister VInDelay = VectorLoadAligned(&InDelaySamples[i]);
-		VectorRegister VInSamples = VectorLoadAligned(&InSamples[i]);
+		VectorRegister4f VInDelay = VectorLoadAligned(&InDelaySamples[i]);
+		VectorRegister4f VInSamples = VectorLoadAligned(&InSamples[i]);
 		// w[n] = x[n] + G * w[n - D]
-		VectorRegister VOutDelay = VectorMultiplyAdd(VInDelay, VG, VInSamples);
+		VectorRegister4f VOutDelay = VectorMultiplyAdd(VInDelay, VG, VInSamples);
 		
 		// Underflow clamp
-		VectorRegister Mask = VectorBitwiseAnd(VectorCompareGT(VOutDelay, VNFMIN), VectorCompareLT(VOutDelay, VFMIN));
+		VectorRegister4f Mask = VectorBitwiseAnd(VectorCompareGT(VOutDelay, VNFMIN), VectorCompareLT(VOutDelay, VFMIN));
 		VOutDelay = VectorSelect(Mask, GlobalVectorConstants::FloatZero, VOutDelay);
 		VectorStoreAligned(VOutDelay, &OutDelaySamples[i]);
 
 		// y[n] = -G * w[n] + w[n - D]
-		VectorRegister VOut = VectorMultiplyAdd(VOutDelay, VNG, VInDelay);
+		VectorRegister4f VOut = VectorMultiplyAdd(VOutDelay, VNG, VInDelay);
 		VectorStoreAligned(VOut, &OutSamples[i]);
 	}
 
diff --git a/Engine/Source/Runtime/SignalProcessing/Private/UniformPartitionConvolution.cpp b/Engine/Source/Runtime/SignalProcessing/Private/UniformPartitionConvolution.cpp
--- a/Engine/Source/Runtime/SignalProcessing/Private/UniformPartitionConvolution.cpp
+++ b/Engine/Source/Runtime/SignalProcessing/Private/UniformPartitionConvolution.cpp
@@ -265,7 +265,7 @@ namespace Audio
 		const float* InBData = InB.GetData();
 		float* OutData = Out.GetData();
 
-		const VectorRegister SignFlip = MakeVectorRegister(-1.f, 1.f, -1.f, 1.f);
+		const VectorRegister4f SignFlip = MakeVectorRegisterFloat(-1.f, 1.f, -1.f, 1.f);
 
 		for (int32 i = 0; i < NumSimd; i += 4)
 		{
@@ -279,16 +279,16 @@ namespace Audio
 			// B1r B1i B2r B2i
 
 			// VectorA = A1r A1i A2r A2i
-			VectorRegister VectorInA = VectorLoadAligned(&InAData[i]);
+			VectorRegister4f VectorInA = VectorLoadAligned(&InAData[i]);
 			// Temp12 = A1i A1r A2i A2r
-			VectorRegister Temp1 = VectorSwizzle(VectorInA, 1, 0, 3, 2);
+			VectorRegister4f Temp1 = VectorSwizzle(VectorInA, 1, 0, 3, 2);
 
 			// VectorB = B1r B1i B2r B2i
-			VectorRegister VectorInB = VectorLoadAligned(&InBData[i]);
+			VectorRegister4f VectorInB = VectorLoadAligned(&InBData[i]);
 			// Temp2 = B1r B1r B2r B2r
-			VectorRegister Temp2 = VectorSwizzle(VectorInB, 0, 0, 2, 2);
+			VectorRegister4f Temp2 = VectorSwizzle(VectorInB, 0, 0, 2, 2);
 			// Temp3 = B1i B1i B2i B2i
-			VectorRegister Temp3 = VectorSwizzle(VectorInB, 1, 1, 3, 3);
+			VectorRegister4f Temp3 = VectorSwizzle(VectorInB, 1, 1, 3, 3);
 
 
 			// VectorA = A1rB1r, A1iB1r, A2rB2r, A2iB2r
@@ -302,7 +302,7 @@ namespace Audio
 			Temp1 = VectorMultiplyAdd(Temp1, SignFlip, VectorInA);
 
 			// VectorOut = O1r + A1rB1r - A1iB1i, O1i + A1iB1r + A1rB1i, O2r + A2rB2r - A2iB2i, O2i + A2iB2r + A2rB2i
-			VectorRegister VectorOut = VectorLoadAligned(&OutData[i]);
+			VectorRegister4f VectorOut = VectorLoadAligned(&OutData[i]);
 			VectorOut = VectorAdd(Temp1, VectorOut);
 
 			VectorStoreAligned(VectorOut, &OutData[i]);
@@ -322,7 +322,7 @@ namespace Audio
 	{
 		check(InBuffer.Num() == OutBuffer.Num());
 
-		VectorRegister VectorConstant = MakeVectorRegister(InConstant, InConstant, InConstant, InConstant);
+		VectorRegister4f VectorConstant = MakeVectorRegisterFloat(InConstant, InConstant, InConstant, InConstant);
 
 		const int32 Num = InBuffer.Num();
 		const int32 NumSimd = Num & NumSimdMask;
@@ -332,9 +332,9 @@ namespace Audio
 
 		for (int32 i = 0; i < NumSimd; i += 4)
 		{
-			VectorRegister VectorIn = VectorLoadAligned(&InData[i]);
+			VectorRegister4f VectorIn = VectorLoadAligned(&InData[i]);
 
-			VectorRegister VectorOut = VectorMultiply(VectorIn, VectorConstant);
+			VectorRegister4f VectorOut = VectorMultiply(VectorIn, VectorConstant);
 			
 			VectorStoreAligned(VectorOut, &OutData[i]);
 		}
diff --git a/Engine/Source/Runtime/SignalProcessing/Private/VectorFFT.cpp b/Engine/Source/Runtime/SignalProcessing/Private/VectorFFT.cpp
--- a/Engine/Source/Runtime/SignalProcessing/Private/VectorFFT.cpp
+++ b/Engine/Source/Runtime/SignalProcessing/Private/VectorFFT.cpp
@@ -149,27 +149,27 @@ namespace Audio
 			// Structure to hold loaded inputs on final radix 4 pass
 			struct FFinalInputs
 			{
-				VectorRegister A0;
-				VectorRegister A1;
-				VectorRegister A2;
-				VectorRegister A3;
-				VectorRegister A4;
-				VectorRegister A5;
-				VectorRegister A6;
-				VectorRegister A7;
+				VectorRegister4f A0;
+				VectorRegister4f A1;
+				VectorRegister4f A2;
+				VectorRegister4f A3;
+				VectorRegister4f A4;
+				VectorRegister4f A5;
+				VectorRegister4f A6;
+				VectorRegister4f A7;
 			};
 
 			// Structure to hold loaded outputs on final radix 4 pass
 			struct FFinalOutputs
 			{
-				VectorRegister D0;
-				VectorRegister D1;
-				VectorRegister D2;
-				VectorRegister D3;
-				VectorRegister D4;
-				VectorRegister D5;
-				VectorRegister D6;
-				VectorRegister D7;
+				VectorRegister4f D0;
+				VectorRegister4f D1;
+				VectorRegister4f D2;
+				VectorRegister4f D3;
+				VectorRegister4f D4;
+				VectorRegister4f D5;
+				VectorRegister4f D6;
+				VectorRegister4f D7;
 			};
 
 			// Weight structure for final Radix4 pass
@@ -226,12 +226,12 @@ namespace Audio
 				// Complex values in a vector are [real_1, imag_1, real_2, imag_2].  
 				// By multipling this value, we flip the sign of the imaginary components, which
 				// is the equivalent of a complex conjugate.
-				const VectorRegister SignFlipImag = MakeVectorRegister(Scale, -Scale, Scale, -Scale);
+				const VectorRegister4f SignFlipImag = MakeVectorRegisterFloat(Scale, -Scale, Scale, -Scale);
 
 				// Perform operation using SIMD
 				for (int32 i = 0; i < NumToSimd; i += 4)
 				{
-					VectorRegister Value = VectorLoadAligned(&InValues[i]);
+					VectorRegister4f Value = VectorLoadAligned(&InValues[i]);
 					Value = VectorMultiply(SignFlipImag, Value);
 					VectorStoreAligned(Value, &OutValues[i]);
 				}
@@ -258,7 +258,7 @@ namespace Audio
 				// Calculate number of constant weight butterflies in this stage.
 				const int NumButterflies = 1 << (InStageIndex - 2);
 
-				const VectorRegister SignFlipImag = MakeVectorRegister(1.f, -1.f, 1.f, -1.f);
+				const VectorRegister4f SignFlipImag = MakeVectorRegisterFloat(1.f, -1.f, 1.f, -1.f);
 
 				const int32 Offset0 = 0;
 				const int32 Offset1 = 2 * NumButterflies;
@@ -273,23 +273,23 @@ namespace Audio
 					const int32 Pos2 = Offset2 + Pos;
 					const int32 Pos3 = Offset3 + Pos;
 
-					VectorRegister A0 = VectorLoadAligned(&InValues[Pos0]);
-					VectorRegister A1 = VectorLoadAligned(&InValues[Pos1]);
-					VectorRegister A2 = VectorLoadAligned(&InValues[Pos2]);
-					VectorRegister A3 = VectorLoadAligned(&InValues[Pos3]);
+					VectorRegister4f A0 = VectorLoadAligned(&InValues[Pos0]);
+					VectorRegister4f A1 = VectorLoadAligned(&InValues[Pos1]);
+					VectorRegister4f A2 = VectorLoadAligned(&InValues[Pos2]);
+					VectorRegister4f A3 = VectorLoadAligned(&InValues[Pos3]);
 
-					VectorRegister C0 = VectorAdd(A0, A2);
-					VectorRegister C2 = VectorSubtract(A0, A2);
-					VectorRegister C1 = VectorAdd(A1, A3);
-					VectorRegister C3 = VectorSubtract(A1, A3);
+					VectorRegister4f C0 = VectorAdd(A0, A2);
+					VectorRegister4f C2 = VectorSubtract(A0, A2);
+					VectorRegister4f C1 = VectorAdd(A1, A3);
+					VectorRegister4f C3 = VectorSubtract(A1, A3);
 
-					VectorRegister C3Conj = VectorMultiply(C3, SignFlipImag);
-					VectorRegister C3ConjSwizzle = VectorSwizzle(C3Conj, 1, 0, 3, 2);
+					VectorRegister4f C3Conj = VectorMultiply(C3, SignFlipImag);
+					VectorRegister4f C3ConjSwizzle = VectorSwizzle(C3Conj, 1, 0, 3, 2);
 
-					VectorRegister D0 = VectorAdd(C1, C0);
-					VectorRegister D1 = VectorSubtract(C0, C1);
-					VectorRegister D2 = VectorAdd(C2, C3ConjSwizzle);
-					VectorRegister D3 = VectorSubtract(C2, C3ConjSwizzle);
+					VectorRegister4f D0 = VectorAdd(C1, C0);
+					VectorRegister4f D1 = VectorSubtract(C0, C1);
+					VectorRegister4f D2 = VectorAdd(C2, C3ConjSwizzle);
+					VectorRegister4f D3 = VectorSubtract(C2, C3ConjSwizzle);
 
 					VectorStoreAligned(D0, &OutValues[Pos0]);
 					VectorStoreAligned(D1, &OutValues[Pos1]);
@@ -315,9 +315,9 @@ namespace Audio
 				// Calculate number of constant weight butterflies in this stage.
 				const int32 NumButterflies = 1 << (InStageIndex - 3);
 
-				const VectorRegister SignFlipImag = MakeVectorRegister(1.f, -1.f, 1.f, -1.f);
-				const VectorRegister VectorSqrt2D2 = MakeVectorRegister(Sqrt2D2, Sqrt2D2, Sqrt2D2, Sqrt2D2);
-				const VectorRegister VectorNegSqrt2D2 = MakeVectorRegister(-Sqrt2D2, -Sqrt2D2, -Sqrt2D2, -Sqrt2D2);
+				const VectorRegister4f SignFlipImag = MakeVectorRegisterFloat(1.f, -1.f, 1.f, -1.f);
+				const VectorRegister4f VectorSqrt2D2 = MakeVectorRegisterFloat(Sqrt2D2, Sqrt2D2, Sqrt2D2, Sqrt2D2);
+				const VectorRegister4f VectorNegSqrt2D2 = MakeVectorRegisterFloat(-Sqrt2D2, -Sqrt2D2, -Sqrt2D2, -Sqrt2D2);
 
 				const int32 Offset0 = 0;
 				const int32 Offset1 = 2 * NumButterflies;
@@ -340,59 +340,59 @@ namespace Audio
 					const int32 Pos6 = Offset6 + Pos;
 					const int32 Pos7 = Offset7 + Pos;
 
-					VectorRegister A0 = VectorLoadAligned(&InValues[Pos0]);
-					VectorRegister A1 = VectorLoadAligned(&InValues[Pos1]);
-					VectorRegister A2 = VectorLoadAligned(&InValues[Pos2]);
-					VectorRegister A3 = VectorLoadAligned(&InValues[Pos3]);
-					VectorRegister A4 = VectorLoadAligned(&InValues[Pos4]);
-					VectorRegister A5 = VectorLoadAligned(&InValues[Pos5]);
-					VectorRegister A6 = VectorLoadAligned(&InValues[Pos6]);
-					VectorRegister A7 = VectorLoadAligned(&InValues[Pos7]);
-
-					VectorRegister B0 = VectorAdd(A0, A4);
-					VectorRegister B1 = VectorAdd(A1, A5);
-					VectorRegister B2 = VectorAdd(A2, A6);
-					VectorRegister B3 = VectorAdd(A3, A7);
-					VectorRegister B4 = VectorSubtract(A0, A4);
-					VectorRegister B5 = VectorSubtract(A1, A5);
-					VectorRegister B6 = VectorSubtract(A2, A6);
-					VectorRegister B7 = VectorSubtract(A3, A7);
-
-					VectorRegister B6Conj = VectorMultiply(SignFlipImag, B6);
-					VectorRegister B6ConjSwizzle = VectorSwizzle(B6Conj, 1, 0, 3, 2);
-
-					VectorRegister B7Conj = VectorMultiply(SignFlipImag, B7);
-					VectorRegister B7ConjSwizzle = VectorSwizzle(B7Conj, 1, 0, 3, 2);
-
-					VectorRegister C0 = VectorAdd(B0, B2);
-					VectorRegister C1 = VectorAdd(B1, B3);
-					VectorRegister C2 = VectorSubtract(B0, B2);
-					VectorRegister C3 = VectorSubtract(B1, B3);
-					VectorRegister C4 = VectorAdd(B4, B6ConjSwizzle);
-					VectorRegister C5 = VectorAdd(B5, B7ConjSwizzle);
-					VectorRegister C6 = VectorSubtract(B4, B6ConjSwizzle);
-					VectorRegister C7 = VectorSubtract(B5, B7ConjSwizzle);
-
-					VectorRegister C3Conj = VectorMultiply(SignFlipImag, C3);
-					VectorRegister C3ConjSwizzle = VectorSwizzle(C3Conj, 1, 0, 3, 2);
-
-					VectorRegister C5Conj = VectorMultiply(SignFlipImag, C5);
-					VectorRegister C5ConjSwizzle = VectorSwizzle(C5Conj, 1, 0, 3, 2);
-					VectorRegister T5 = VectorAdd(C5, C5ConjSwizzle);
-
-					VectorRegister C7Swizzle = VectorSwizzle(C7, 1, 0, 3, 2);
-
-					VectorRegister T7 = VectorMultiplyAdd(SignFlipImag, C7, C7Swizzle);
-					VectorRegister T7Conj = VectorMultiply(T7, SignFlipImag);
-
-					VectorRegister D0 = VectorAdd(C0, C1);
-					VectorRegister D1 = VectorSubtract(C0, C1);
-					VectorRegister D2 = VectorAdd(C2, C3ConjSwizzle);
-					VectorRegister D3 = VectorSubtract(C2, C3ConjSwizzle);
-					VectorRegister D4 = VectorMultiplyAdd(T5, VectorSqrt2D2, C4);
-					VectorRegister D5 = VectorMultiplyAdd(T5, VectorNegSqrt2D2, C4);
-					VectorRegister D6 = VectorMultiplyAdd(VectorNegSqrt2D2, T7Conj, C6);
-					VectorRegister D7 = VectorMultiplyAdd(VectorSqrt2D2, T7Conj, C6);
+					VectorRegister4f A0 = VectorLoadAligned(&InValues[Pos0]);
+					VectorRegister4f A1 = VectorLoadAligned(&InValues[Pos1]);
+					VectorRegister4f A2 = VectorLoadAligned(&InValues[Pos2]);
+					VectorRegister4f A3 = VectorLoadAligned(&InValues[Pos3]);
+					VectorRegister4f A4 = VectorLoadAligned(&InValues[Pos4]);
+					VectorRegister4f A5 = VectorLoadAligned(&InValues[Pos5]);
+					VectorRegister4f A6 = VectorLoadAligned(&InValues[Pos6]);
+					VectorRegister4f A7 = VectorLoadAligned(&InValues[Pos7]);
+
+					VectorRegister4f B0 = VectorAdd(A0, A4);
+					VectorRegister4f B1 = VectorAdd(A1, A5);
+					VectorRegister4f B2 = VectorAdd(A2, A6);
+					VectorRegister4f B3 = VectorAdd(A3, A7);
+					VectorRegister4f B4 = VectorSubtract(A0, A4);
+					VectorRegister4f B5 = VectorSubtract(A1, A5);
+					VectorRegister4f B6 = VectorSubtract(A2, A6);
+					VectorRegister4f B7 = VectorSubtract(A3, A7);
+
+					VectorRegister4f B6Conj = VectorMultiply(SignFlipImag, B6);
+					VectorRegister4f B6ConjSwizzle = VectorSwizzle(B6Conj, 1, 0, 3, 2);
+
+					VectorRegister4f B7Conj = VectorMultiply(SignFlipImag, B7);
+					VectorRegister4f B7ConjSwizzle = VectorSwizzle(B7Conj, 1, 0, 3, 2);
+
+					VectorRegister4f C0 = VectorAdd(B0, B2);
+					VectorRegister4f C1 = VectorAdd(B1, B3);
+					VectorRegister4f C2 = VectorSubtract(B0, B2);
+					VectorRegister4f C3 = VectorSubtract(B1, B3);
+					VectorRegister4f C4 = VectorAdd(B4, B6ConjSwizzle);
+					VectorRegister4f C5 = VectorAdd(B5, B7ConjSwizzle);
+					VectorRegister4f C6 = VectorSubtract(B4, B6ConjSwizzle);
+					VectorRegister4f C7 = VectorSubtract(B5, B7ConjSwizzle);
+
+					VectorRegister4f C3Conj = VectorMultiply(SignFlipImag, C3);
+					VectorRegister4f C3ConjSwizzle = VectorSwizzle(C3Conj, 1, 0, 3, 2);
+
+					VectorRegister4f C5Conj = VectorMultiply(SignFlipImag, C5);
+					VectorRegister4f C5ConjSwizzle = VectorSwizzle(C5Conj, 1, 0, 3, 2);
+					VectorRegister4f T5 = VectorAdd(C5, C5ConjSwizzle);
+
+					VectorRegister4f C7Swizzle = VectorSwizzle(C7, 1, 0, 3, 2);
+
+					VectorRegister4f T7 = VectorMultiplyAdd(SignFlipImag, C7, C7Swizzle);
+					VectorRegister4f T7Conj = VectorMultiply(T7, SignFlipImag);
+
+					VectorRegister4f D0 = VectorAdd(C0, C1);
+					VectorRegister4f D1 = VectorSubtract(C0, C1);
+					VectorRegister4f D2 = VectorAdd(C2, C3ConjSwizzle);
+					VectorRegister4f D3 = VectorSubtract(C2, C3ConjSwizzle);
+					VectorRegister4f D4 = VectorMultiplyAdd(T5, VectorSqrt2D2, C4);
+					VectorRegister4f D5 = VectorMultiplyAdd(T5, VectorNegSqrt2D2, C4);
+					VectorRegister4f D6 = VectorMultiplyAdd(VectorNegSqrt2D2, T7Conj, C6);
+					VectorRegister4f D7 = VectorMultiplyAdd(VectorSqrt2D2, T7Conj, C6);
 
 					VectorStoreAligned(D0, &OutValues[Pos0]);
 					VectorStoreAligned(D1, &OutValues[Pos1]);
@@ -414,18 +414,18 @@ namespace Audio
 				const int32 NumButterflies = 1 << (InStageIndex - 2);
 
 				// Load weights for butterfly
-				const VectorRegister Weight1Real = VectorLoadAligned(Weights.W1R);
-				const VectorRegister Weight1Imag = VectorLoadAligned(Weights.W1I);
-				const VectorRegister Weight2Real = VectorLoadAligned(Weights.W2R);
-				const VectorRegister Weight2Imag = VectorLoadAligned(Weights.W2I);
-				const VectorRegister Weight3Real = VectorLoadAligned(Weights.W3R);
-				const VectorRegister Weight3Imag = VectorLoadAligned(Weights.W3I);
-
-				const VectorRegister Weight3RealNeg = VectorLoadAligned(Weights.W3RNeg);
-				const VectorRegister Weight2RealNeg = VectorLoadAligned(Weights.W2RNeg);
-				const VectorRegister Weight1RealNeg = VectorLoadAligned(Weights.W1RNeg);
-				const VectorRegister Weight1RealD2 = VectorLoadAligned(Weights.W1RD2);
-				const VectorRegister Weight1RealD3 = VectorLoadAligned(Weights.W1RD3);
+				const VectorRegister4f Weight1Real = VectorLoadAligned(Weights.W1R);
+				const VectorRegister4f Weight1Imag = VectorLoadAligned(Weights.W1I);
+				const VectorRegister4f Weight2Real = VectorLoadAligned(Weights.W2R);
+				const VectorRegister4f Weight2Imag = VectorLoadAligned(Weights.W2I);
+				const VectorRegister4f Weight3Real = VectorLoadAligned(Weights.W3R);
+				const VectorRegister4f Weight3Imag = VectorLoadAligned(Weights.W3I);
+
+				const VectorRegister4f Weight3RealNeg = VectorLoadAligned(Weights.W3RNeg);
+				const VectorRegister4f Weight2RealNeg = VectorLoadAligned(Weights.W2RNeg);
+				const VectorRegister4f Weight1RealNeg = VectorLoadAligned(Weights.W1RNeg);
+				const VectorRegister4f Weight1RealD2 = VectorLoadAligned(Weights.W1RD2);
+				const VectorRegister4f Weight1RealD3 = VectorLoadAligned(Weights.W1RD3);
 
 				// Perform butterflies.
 				for (int32 i = 0; i < NumButterflies; i += 2)
@@ -435,30 +435,30 @@ namespace Audio
 					const int32 Pos2 = 2 * (Stride * ButterflyIndex + 2 * NumButterflies + i);
 					const int32 Pos3 = 2 * (Stride * ButterflyIndex + 3 * NumButterflies + i);
 
-					VectorRegister A0 = VectorLoadAligned(&InOutValues[Pos0]);
-					VectorRegister A1 = VectorLoadAligned(&InOutValues[Pos1]);
-					VectorRegister A2 = VectorLoadAligned(&InOutValues[Pos2]);
-					VectorRegister A3 = VectorLoadAligned(&InOutValues[Pos3]);
+					VectorRegister4f A0 = VectorLoadAligned(&InOutValues[Pos0]);
+					VectorRegister4f A1 = VectorLoadAligned(&InOutValues[Pos1]);
+					VectorRegister4f A2 = VectorLoadAligned(&InOutValues[Pos2]);
+					VectorRegister4f A3 = VectorLoadAligned(&InOutValues[Pos3]);
 
-					VectorRegister A1Swizzle = VectorSwizzle(A1, 1, 0, 3, 2);
-					VectorRegister A2Swizzle = VectorSwizzle(A2, 1, 0, 3, 2);
-					VectorRegister A3Swizzle = VectorSwizzle(A3, 1, 0, 3, 2);
+					VectorRegister4f A1Swizzle = VectorSwizzle(A1, 1, 0, 3, 2);
+					VectorRegister4f A2Swizzle = VectorSwizzle(A2, 1, 0, 3, 2);
+					VectorRegister4f A3Swizzle = VectorSwizzle(A3, 1, 0, 3, 2);
 
-					VectorRegister B1 = VectorMultiplyAdd(A1Swizzle, Weight1Imag, A1);
-					VectorRegister B2 = VectorMultiplyAdd(A2Swizzle, Weight2Imag, A2);
-					VectorRegister B3 = VectorMultiplyAdd(A3Swizzle, Weight3Imag, A3);
+					VectorRegister4f B1 = VectorMultiplyAdd(A1Swizzle, Weight1Imag, A1);
+					VectorRegister4f B2 = VectorMultiplyAdd(A2Swizzle, Weight2Imag, A2);
+					VectorRegister4f B3 = VectorMultiplyAdd(A3Swizzle, Weight3Imag, A3);
 
-					VectorRegister C0 = VectorMultiplyAdd(B2, Weight2Real, A0);
-					VectorRegister C2 = VectorMultiplyAdd(B2, Weight2RealNeg, A0);
-					VectorRegister C1 = VectorMultiplyAdd(B3, Weight3Real, B1);
-					VectorRegister C3 = VectorMultiplyAdd(B3, Weight3RealNeg, B1);
+					VectorRegister4f C0 = VectorMultiplyAdd(B2, Weight2Real, A0);
+					VectorRegister4f C2 = VectorMultiplyAdd(B2, Weight2RealNeg, A0);
+					VectorRegister4f C1 = VectorMultiplyAdd(B3, Weight3Real, B1);
+					VectorRegister4f C3 = VectorMultiplyAdd(B3, Weight3RealNeg, B1);
 
-					VectorRegister C3Swizzle = VectorSwizzle(C3, 1, 0, 3, 2);
+					VectorRegister4f C3Swizzle = VectorSwizzle(C3, 1, 0, 3, 2);
 
-					VectorRegister D0 = VectorMultiplyAdd(C1, Weight1Real, C0);
-					VectorRegister D1 = VectorMultiplyAdd(C1, Weight1RealNeg, C0);
-					VectorRegister D2 = VectorMultiplyAdd(C3Swizzle, Weight1RealD2, C2);
-					VectorRegister D3 = VectorMultiplyAdd(C3Swizzle, Weight1RealD3, C2);
+					VectorRegister4f D0 = VectorMultiplyAdd(C1, Weight1Real, C0);
+					VectorRegister4f D1 = VectorMultiplyAdd(C1, Weight1RealNeg, C0);
+					VectorRegister4f D2 = VectorMultiplyAdd(C3Swizzle, Weight1RealD2, C2);
+					VectorRegister4f D3 = VectorMultiplyAdd(C3Swizzle, Weight1RealD3, C2);
 
 					VectorStoreAligned(D0, &InOutValues[Pos0]);
 					VectorStoreAligned(D1, &InOutValues[Pos1]);
@@ -481,18 +481,18 @@ namespace Audio
 					// Load values for current weight.
 					const FRadix4Weight& Weights = Radix4Weights[i];
 
-					const VectorRegister Weight1Real = VectorLoadAligned(Weights.W1R);
-					const VectorRegister Weight1Imag = VectorLoadAligned(Weights.W1I);
-					const VectorRegister Weight2Real = VectorLoadAligned(Weights.W2R);
-					const VectorRegister Weight2Imag = VectorLoadAligned(Weights.W2I);
-					const VectorRegister Weight3Real = VectorLoadAligned(Weights.W3R);
-					const VectorRegister Weight3Imag = VectorLoadAligned(Weights.W3I);
+					const VectorRegister4f Weight1Real = VectorLoadAligned(Weights.W1R);
+					const VectorRegister4f Weight1Imag = VectorLoadAligned(Weights.W1I);
+					const VectorRegister4f Weight2Real = VectorLoadAligned(Weights.W2R);
+					const VectorRegister4f Weight2Imag = VectorLoadAligned(Weights.W2I);
+					const VectorRegister4f Weight3Real = VectorLoadAligned(Weights.W3R);
+					const VectorRegister4f Weight3Imag = VectorLoadAligned(Weights.W3I);
 
-					const VectorRegister Weight3RealNeg = VectorLoadAligned(Weights.W3RNeg);
-					const VectorRegister Weight2RealNeg = VectorLoadAligned(Weights.W2RNeg);
-					const VectorRegister Weight1RealNeg = VectorLoadAligned(Weights.W1RNeg);
-					const VectorRegister Weight1RealD2 = VectorLoadAligned(Weights.W1RD2);
-					const VectorRegister Weight1RealD3 = VectorLoadAligned(Weights.W1RD3);
+					const VectorRegister4f Weight3RealNeg = VectorLoadAligned(Weights.W3RNeg);
+					const VectorRegister4f Weight2RealNeg = VectorLoadAligned(Weights.W2RNeg);
+					const VectorRegister4f Weight1RealNeg = VectorLoadAligned(Weights.W1RNeg);
+					const VectorRegister4f Weight1RealD2 = VectorLoadAligned(Weights.W1RD2);
+					const VectorRegister4f Weight1RealD3 = VectorLoadAligned(Weights.W1RD3);
 
 					for (int32 j = 0; j < 4 ; j += 2)
 					{
@@ -501,29 +501,29 @@ namespace Audio
 						const int32 Pos2 = 2 * (Stride * i + 8 + j);
 						const int32 Pos3 = 2 * (Stride * i + 12 + j);
 
-						VectorRegister A0 = VectorLoadAligned(&InOutValues[Pos0]);
-						VectorRegister A1 = VectorLoadAligned(&InOutValues[Pos1]);
-						VectorRegister A2 = VectorLoadAligned(&InOutValues[Pos2]);
-						VectorRegister A3 = VectorLoadAligned(&InOutValues[Pos3]);
+						VectorRegister4f A0 = VectorLoadAligned(&InOutValues[Pos0]);
+						VectorRegister4f A1 = VectorLoadAligned(&InOutValues[Pos1]);
+						VectorRegister4f A2 = VectorLoadAligned(&InOutValues[Pos2]);
+						VectorRegister4f A3 = VectorLoadAligned(&InOutValues[Pos3]);
 
-						VectorRegister A1Swizzle = VectorSwizzle(A1, 1, 0, 3, 2);
-						VectorRegister A2Swizzle = VectorSwizzle(A2, 1, 0, 3, 2);
-						VectorRegister A3Swizzle = VectorSwizzle(A3, 1, 0, 3, 2);
+						VectorRegister4f A1Swizzle = VectorSwizzle(A1, 1, 0, 3, 2);
+						VectorRegister4f A2Swizzle = VectorSwizzle(A2, 1, 0, 3, 2);
+						VectorRegister4f A3Swizzle = VectorSwizzle(A3, 1, 0, 3, 2);
 
-						VectorRegister B1 = VectorMultiplyAdd(A1Swizzle, Weight1Imag, A1);
-						VectorRegister B2 = VectorMultiplyAdd(A2Swizzle, Weight2Imag, A2);
-						VectorRegister B3 = VectorMultiplyAdd(A3Swizzle, Weight3Imag, A3);
+						VectorRegister4f B1 = VectorMultiplyAdd(A1Swizzle, Weight1Imag, A1);
+						VectorRegister4f B2 = VectorMultiplyAdd(A2Swizzle, Weight2Imag, A2);
+						VectorRegister4f B3 = VectorMultiplyAdd(A3Swizzle, Weight3Imag, A3);
 						
-						VectorRegister C0 = VectorMultiplyAdd(B2, Weight2Real, A0);
-						VectorRegister C1 = VectorMultiplyAdd(B3, Weight3Real, B1);
-						VectorRegister C2 = VectorMultiplyAdd(B2, Weight2RealNeg, A0);
-						VectorRegister C3 = VectorMultiplyAdd(B3, Weight3RealNeg, B1);
-						VectorRegister C3Swizzle = VectorSwizzle(C3, 1, 0, 3, 2);
+						VectorRegister4f C0 = VectorMultiplyAdd(B2, Weight2Real, A0);
+						VectorRegister4f C1 = VectorMultiplyAdd(B3, Weight3Real, B1);
+						VectorRegister4f C2 = VectorMultiplyAdd(B2, Weight2RealNeg, A0);
+						VectorRegister4f C3 = VectorMultiplyAdd(B3, Weight3RealNeg, B1);
+						VectorRegister4f C3Swizzle = VectorSwizzle(C3, 1, 0, 3, 2);
 						
-						VectorRegister D0 = VectorMultiplyAdd(C1, Weight1Real, C0);
-						VectorRegister D1 = VectorMultiplyAdd(C1, Weight1RealNeg, C0);
-						VectorRegister D2 = VectorMultiplyAdd(C3Swizzle, Weight1RealD2, C2);
-						VectorRegister D3 = VectorMultiplyAdd(C3Swizzle, Weight1RealD3, C2);
+						VectorRegister4f D0 = VectorMultiplyAdd(C1, Weight1Real, C0);
+						VectorRegister4f D1 = VectorMultiplyAdd(C1, Weight1RealNeg, C0);
+						VectorRegister4f D2 = VectorMultiplyAdd(C3Swizzle, Weight1RealD2, C2);
+						VectorRegister4f D3 = VectorMultiplyAdd(C3Swizzle, Weight1RealD3, C2);
 
 						VectorStoreAligned(D0, &InOutValues[Pos0]);
 						VectorStoreAligned(D1, &InOutValues[Pos1]);
@@ -546,14 +546,14 @@ namespace Audio
 				const int32 Pos6 = 2 * (3 * InNumButterflies + 4 * InReadIndex);
 				const int32 Pos7 = Pos6 + 4;
 				
-				VectorRegister T0 = VectorLoadAligned(&InValues[Pos0]);
-				VectorRegister T1 = VectorLoadAligned(&InValues[Pos1]);
-				VectorRegister T2 = VectorLoadAligned(&InValues[Pos2]);
-				VectorRegister T3 = VectorLoadAligned(&InValues[Pos3]);
-				VectorRegister T4 = VectorLoadAligned(&InValues[Pos4]);
-				VectorRegister T5 = VectorLoadAligned(&InValues[Pos5]);
-				VectorRegister T6 = VectorLoadAligned(&InValues[Pos6]);
-				VectorRegister T7 = VectorLoadAligned(&InValues[Pos7]);
+				VectorRegister4f T0 = VectorLoadAligned(&InValues[Pos0]);
+				VectorRegister4f T1 = VectorLoadAligned(&InValues[Pos1]);
+				VectorRegister4f T2 = VectorLoadAligned(&InValues[Pos2]);
+				VectorRegister4f T3 = VectorLoadAligned(&InValues[Pos3]);
+				VectorRegister4f T4 = VectorLoadAligned(&InValues[Pos4]);
+				VectorRegister4f T5 = VectorLoadAligned(&InValues[Pos5]);
+				VectorRegister4f T6 = VectorLoadAligned(&InValues[Pos6]);
+				VectorRegister4f T7 = VectorLoadAligned(&InValues[Pos7]);
 
 				OutValues.A0 = VectorShuffle(T0, T2, 0, 1, 0, 1);
 				OutValues.A1 = VectorShuffle(T4, T6, 0, 1, 0, 1);
@@ -592,56 +592,56 @@ namespace Audio
 			void Radix4ButterflyFinalIteration(const FFinalInputs& Inputs, const FFinalWeights& InWeights, FFinalOutputs& Outputs)
 			{
 				// Note: Some weights are altered to bake in sign flips to avoid an extra multiply later on.
-				const VectorRegister W2I = VectorLoadAligned(InWeights.W2I);
-				const VectorRegister W2R = VectorLoadAligned(InWeights.W2R);
-				const VectorRegister W3I = VectorLoadAligned(InWeights.W3I);
-				const VectorRegister W3R = VectorLoadAligned(InWeights.W3R);
-				const VectorRegister W4I = VectorLoadAligned(InWeights.W4I);
-				const VectorRegister W4R = VectorLoadAligned(InWeights.W4R);
-				const VectorRegister W5I = VectorLoadAligned(InWeights.W5I);
-				const VectorRegister W5R = VectorLoadAligned(InWeights.W5R);
-				const VectorRegister W6I = VectorLoadAligned(InWeights.W6I);
-				const VectorRegister W6R = VectorLoadAligned(InWeights.W6R);
-				const VectorRegister W7I = VectorLoadAligned(InWeights.W7I);
-				const VectorRegister W7R = VectorLoadAligned(InWeights.W7R);
-
-				const VectorRegister W2RNeg = VectorLoadAligned(InWeights.W2RNeg);
-				const VectorRegister W3RNeg = VectorLoadAligned(InWeights.W3RNeg);
-				const VectorRegister W4RNeg = VectorLoadAligned(InWeights.W4RNeg);
-				const VectorRegister W5RNeg = VectorLoadAligned(InWeights.W5RNeg);
-				const VectorRegister W6RNeg = VectorLoadAligned(InWeights.W6RNeg);
-				const VectorRegister W7RNeg = VectorLoadAligned(InWeights.W7RNeg);
-
-				const VectorRegister W2RD4 = VectorLoadAligned(InWeights.W2RD4);
-				const VectorRegister W2RD6 = VectorLoadAligned(InWeights.W2RD6);
-				const VectorRegister W3RD5 = VectorLoadAligned(InWeights.W3RD5);
-				const VectorRegister W3RD7 = VectorLoadAligned(InWeights.W3RD7);
-
-				VectorRegister A2Swizzle = VectorSwizzle(Inputs.A2, 1, 0, 3, 2);
-				VectorRegister A3Swizzle = VectorSwizzle(Inputs.A3, 1, 0, 3, 2);
-				VectorRegister A4Swizzle = VectorSwizzle(Inputs.A4, 1, 0, 3, 2);
-				VectorRegister A5Swizzle = VectorSwizzle(Inputs.A5, 1, 0, 3, 2);
-				VectorRegister A6Swizzle = VectorSwizzle(Inputs.A6, 1, 0, 3, 2);
-				VectorRegister A7Swizzle = VectorSwizzle(Inputs.A7, 1, 0, 3, 2);
-
-				VectorRegister B2 = VectorMultiplyAdd(A2Swizzle, W2I, Inputs.A2);
-				VectorRegister B3 = VectorMultiplyAdd(A3Swizzle, W3I, Inputs.A3);
-				VectorRegister B4 = VectorMultiplyAdd(A4Swizzle, W4I, Inputs.A4);
-				VectorRegister B5 = VectorMultiplyAdd(A5Swizzle, W5I, Inputs.A5);
-				VectorRegister B6 = VectorMultiplyAdd(A6Swizzle, W6I, Inputs.A6);
-				VectorRegister B7 = VectorMultiplyAdd(A7Swizzle, W7I, Inputs.A7);
-
-				VectorRegister C0 = VectorMultiplyAdd(B4, W4R, Inputs.A0);
-				VectorRegister C1 = VectorMultiplyAdd(B5, W5R, Inputs.A1);
-				VectorRegister C2 = VectorMultiplyAdd(B6, W6R, B2);
-				VectorRegister C3 = VectorMultiplyAdd(B7, W7R, B3);
-				VectorRegister C4 = VectorMultiplyAdd(B4, W4RNeg, Inputs.A0);
-				VectorRegister C5 = VectorMultiplyAdd(B5, W5RNeg, Inputs.A1);
-				VectorRegister C6 = VectorMultiplyAdd(B6, W6RNeg, B2);
-				VectorRegister C7 = VectorMultiplyAdd(B7, W7RNeg, B3);
-
-				VectorRegister C6Swizzle = VectorSwizzle(C6, 1, 0, 3, 2);
-				VectorRegister C7Swizzle = VectorSwizzle(C7, 1, 0, 3, 2);
+				const VectorRegister4f W2I = VectorLoadAligned(InWeights.W2I);
+				const VectorRegister4f W2R = VectorLoadAligned(InWeights.W2R);
+				const VectorRegister4f W3I = VectorLoadAligned(InWeights.W3I);
+				const VectorRegister4f W3R = VectorLoadAligned(InWeights.W3R);
+				const VectorRegister4f W4I = VectorLoadAligned(InWeights.W4I);
+				const VectorRegister4f W4R = VectorLoadAligned(InWeights.W4R);
+				const VectorRegister4f W5I = VectorLoadAligned(InWeights.W5I);
+				const VectorRegister4f W5R = VectorLoadAligned(InWeights.W5R);
+				const VectorRegister4f W6I = VectorLoadAligned(InWeights.W6I);
+				const VectorRegister4f W6R = VectorLoadAligned(InWeights.W6R);
+				const VectorRegister4f W7I = VectorLoadAligned(InWeights.W7I);
+				const VectorRegister4f W7R = VectorLoadAligned(InWeights.W7R);
+
+				const VectorRegister4f W2RNeg = VectorLoadAligned(InWeights.W2RNeg);
+				const VectorRegister4f W3RNeg = VectorLoadAligned(InWeights.W3RNeg);
+				const VectorRegister4f W4RNeg = VectorLoadAligned(InWeights.W4RNeg);
+				const VectorRegister4f W5RNeg = VectorLoadAligned(InWeights.W5RNeg);
+				const VectorRegister4f W6RNeg = VectorLoadAligned(InWeights.W6RNeg);
+				const VectorRegister4f W7RNeg = VectorLoadAligned(InWeights.W7RNeg);
+
+				const VectorRegister4f W2RD4 = VectorLoadAligned(InWeights.W2RD4);
+				const VectorRegister4f W2RD6 = VectorLoadAligned(InWeights.W2RD6);
+				const VectorRegister4f W3RD5 = VectorLoadAligned(InWeights.W3RD5);
+				const VectorRegister4f W3RD7 = VectorLoadAligned(InWeights.W3RD7);
+
+				VectorRegister4f A2Swizzle = VectorSwizzle(Inputs.A2, 1, 0, 3, 2);
+				VectorRegister4f A3Swizzle = VectorSwizzle(Inputs.A3, 1, 0, 3, 2);
+				VectorRegister4f A4Swizzle = VectorSwizzle(Inputs.A4, 1, 0, 3, 2);
+				VectorRegister4f A5Swizzle = VectorSwizzle(Inputs.A5, 1, 0, 3, 2);
+				VectorRegister4f A6Swizzle = VectorSwizzle(Inputs.A6, 1, 0, 3, 2);
+				VectorRegister4f A7Swizzle = VectorSwizzle(Inputs.A7, 1, 0, 3, 2);
+
+				VectorRegister4f B2 = VectorMultiplyAdd(A2Swizzle, W2I, Inputs.A2);
+				VectorRegister4f B3 = VectorMultiplyAdd(A3Swizzle, W3I, Inputs.A3);
+				VectorRegister4f B4 = VectorMultiplyAdd(A4Swizzle, W4I, Inputs.A4);
+				VectorRegister4f B5 = VectorMultiplyAdd(A5Swizzle, W5I, Inputs.A5);
+				VectorRegister4f B6 = VectorMultiplyAdd(A6Swizzle, W6I, Inputs.A6);
+				VectorRegister4f B7 = VectorMultiplyAdd(A7Swizzle, W7I, Inputs.A7);
+
+				VectorRegister4f C0 = VectorMultiplyAdd(B4, W4R, Inputs.A0);
+				VectorRegister4f C1 = VectorMultiplyAdd(B5, W5R, Inputs.A1);
+				VectorRegister4f C2 = VectorMultiplyAdd(B6, W6R, B2);
+				VectorRegister4f C3 = VectorMultiplyAdd(B7, W7R, B3);
+				VectorRegister4f C4 = VectorMultiplyAdd(B4, W4RNeg, Inputs.A0);
+				VectorRegister4f C5 = VectorMultiplyAdd(B5, W5RNeg, Inputs.A1);
+				VectorRegister4f C6 = VectorMultiplyAdd(B6, W6RNeg, B2);
+				VectorRegister4f C7 = VectorMultiplyAdd(B7, W7RNeg, B3);
+
+				VectorRegister4f C6Swizzle = VectorSwizzle(C6, 1, 0, 3, 2);
+				VectorRegister4f C7Swizzle = VectorSwizzle(C7, 1, 0, 3, 2);
 
 				Outputs.D0 = VectorMultiplyAdd(C2, W2R, C0);
 				Outputs.D1 = VectorMultiplyAdd(C3, W3R, C1);
@@ -1076,43 +1076,43 @@ namespace Audio
 
 		if (FFTSize > InStartIndex)
 		{
-			VectorRegister VInRev1 = VectorLoadAligned(&InValues[FFTSize - InStartIndex]);
+			VectorRegister4f VInRev1 = VectorLoadAligned(&InValues[FFTSize - InStartIndex]);
 
 			for (int32 i = InStartIndex; i < FFTSize; i += 4)
 			{
-				VectorRegister VIn = VectorLoadAligned(&InValues[i]);
-				VectorRegister VInRISwap = VectorSwizzle(VIn, 1, 0, 3, 2);
+				VectorRegister4f VIn = VectorLoadAligned(&InValues[i]);
+				VectorRegister4f VInRISwap = VectorSwizzle(VIn, 1, 0, 3, 2);
 
-				VectorRegister VInRev2 = VectorLoadAligned(&InValues[FFTSize - i - 4]);
-				VectorRegister VInRev = VectorShuffle(VInRev1, VInRev2, 0, 1, 2, 3);
+				VectorRegister4f VInRev2 = VectorLoadAligned(&InValues[FFTSize - i - 4]);
+				VectorRegister4f VInRev = VectorShuffle(VInRev1, VInRev2, 0, 1, 2, 3);
 				VInRev1 = VInRev2;
 
-				VectorRegister VInRevRISwap = VectorSwizzle(VInRev, 1, 0, 3, 2);
+				VectorRegister4f VInRevRISwap = VectorSwizzle(VInRev, 1, 0, 3, 2);
 
-				VectorRegister VAlphaReal = VectorLoadAligned(&AlphaRealData[i]);
-				VectorRegister VAlphaImag = VectorLoadAligned(&AlphaImagData[i]);
-				VectorRegister VBetaReal = VectorLoadAligned(&BetaRealData[i]);
-				VectorRegister VBetaImag = VectorLoadAligned(&BetaImagData[i]);
+				VectorRegister4f VAlphaReal = VectorLoadAligned(&AlphaRealData[i]);
+				VectorRegister4f VAlphaImag = VectorLoadAligned(&AlphaImagData[i]);
+				VectorRegister4f VBetaReal = VectorLoadAligned(&BetaRealData[i]);
+				VectorRegister4f VBetaImag = VectorLoadAligned(&BetaImagData[i]);
 
 				// Out1 = [ R * Ar,  I * Ar]
 				// Out2 = [ I * Ai,  R * Ai]
 				// Out3 = [NR * Br, NI * Br]
 				// Out4 = [NI * Bi, NR * Bi]
-				//VectorRegister Out1 = VectorMultiply(VIn, VAlphaReal);
-				VectorRegister Out2 = VectorMultiply(VInRISwap, VAlphaImag);
-				//VectorRegister Out3 = VectorMultiply(VInRev, VBetaReal);
-				VectorRegister Out4 = VectorMultiply(VInRevRISwap, VBetaImag);
+				//VectorRegister4f Out1 = VectorMultiply(VIn, VAlphaReal);
+				VectorRegister4f Out2 = VectorMultiply(VInRISwap, VAlphaImag);
+				//VectorRegister4f Out3 = VectorMultiply(VInRev, VBetaReal);
+				VectorRegister4f Out4 = VectorMultiply(VInRevRISwap, VBetaImag);
 
 				// Out12 = [(R * Ar) + (I * Ai), (I * Ar) + (R * Ai)]
-				VectorRegister Out12 = VectorMultiplyAdd(VIn, VAlphaReal, Out2);
+				VectorRegister4f Out12 = VectorMultiplyAdd(VIn, VAlphaReal, Out2);
 				// Out34 = [(NR * Br) + (NI * Bi), (NR * Bi) + (NI * Br)]
-				VectorRegister Out34 = VectorMultiplyAdd(VInRev, VBetaReal, Out4);
+				VectorRegister4f Out34 = VectorMultiplyAdd(VInRev, VBetaReal, Out4);
 
 				// Out = [
 				// 	(R * Ar) + (I * Ai) + (NR * Br) + (NI * Bi),
 				// 	(I * Ar) + (R * Ai) + (NR * Bi) + (NI * Br)
 				// ]
-				VectorRegister Out = VectorAdd(Out12, Out34);
+				VectorRegister4f Out = VectorAdd(Out12, Out34);
 				VectorStoreAligned(Out, &OutValues[i]);
 			}
 		}
diff --git a/Engine/Source/Runtime/SignalProcessing/Private/VectorLinearResampler.cpp b/Engine/Source/Runtime/SignalProcessing/Private/VectorLinearResampler.cpp
--- a/Engine/Source/Runtime/SignalProcessing/Private/VectorLinearResampler.cpp
+++ b/Engine/Source/Runtime/SignalProcessing/Private/VectorLinearResampler.cpp
@@ -43,25 +43,25 @@ uint32 Audio::FVectorLinearResampler::ResampleMono(uint32 OutputFramesNeeded, ui
 
 		float CurrentFrameFractionFloat = (float)(CurrentFrameFraction & 0xffff);
 
-		VectorRegister ScaleVec = VectorSetFloat1(1.0f / 65536.0f);
-		VectorRegister TwoVec = VectorSetFloat1(2.0f);
-		VectorRegister OneVec = VectorSetFloat1(1.0f);
-		VectorRegister ThreeVec = VectorSetFloat1(3.0f);
+		VectorRegister4f ScaleVec = VectorSetFloat1(1.0f / 65536.0f);
+		VectorRegister4f TwoVec = VectorSetFloat1(2.0f);
+		VectorRegister4f OneVec = VectorSetFloat1(1.0f);
+		VectorRegister4f ThreeVec = VectorSetFloat1(3.0f);
 
-		VectorRegister CurFracStepVec = VectorSet(
+		VectorRegister4f CurFracStepVec = VectorSet(
 			(float)((0 * FixedPointSampleRate) & 0xffff),
 			(float)((1 * FixedPointSampleRate) & 0xffff),
 			(float)((2 * FixedPointSampleRate) & 0xffff),
 			(float)((3 * FixedPointSampleRate) & 0xffff));
 
-		VectorRegister CurFracVec = VectorSetFloat1(CurrentFrameFractionFloat);
+		VectorRegister4f CurFracVec = VectorSetFloat1(CurrentFrameFractionFloat);
 		CurFracVec = VectorAdd(CurFracVec, CurFracStepVec);
 		CurFracVec = VectorMultiply(CurFracVec, ScaleVec);
 		CurFracVec = VectorAdd(CurFracVec, TwoVec);
 
-		VectorRegister MaskVec = MakeVectorRegister(0xffbfffff, 0xffbfffff, 0xffbfffff, 0xffbfffff);
+		VectorRegister4f MaskVec = MakeVectorRegisterFloat(0xffbfffff, 0xffbfffff, 0xffbfffff, 0xffbfffff);
 
-		VectorRegister StepVec = VectorSetFloat1((float)((FixedPointSampleRate * 4) & 0xffff));
+		VectorRegister4f StepVec = VectorSetFloat1((float)((FixedPointSampleRate * 4) & 0xffff));
 		StepVec = VectorMultiply(StepVec, ScaleVec);
 
 		uint32 OutputFramesNeededSIMD = OutputFramesNeeded & ~3;
@@ -76,22 +76,22 @@ uint32 Audio::FVectorLinearResampler::ResampleMono(uint32 OutputFramesNeeded, ui
 			};
 
 			// [0, 0+1, 1, 1+1]
-			VectorRegister LeftSamples01 = VectorLoadTwoPairsFloat(LeftSamples + SourceOffsets[0], LeftSamples + SourceOffsets[1]);
+			VectorRegister4f LeftSamples01 = VectorLoadTwoPairsFloat(LeftSamples + SourceOffsets[0], LeftSamples + SourceOffsets[1]);
 			// [2, 2+1, 3, 3+1]
-			VectorRegister LeftSamples23 = VectorLoadTwoPairsFloat(LeftSamples + SourceOffsets[2], LeftSamples + SourceOffsets[3]);
+			VectorRegister4f LeftSamples23 = VectorLoadTwoPairsFloat(LeftSamples + SourceOffsets[2], LeftSamples + SourceOffsets[3]);
 
 			// [0, 1, 2, 3]
 			// [0+1, 1+1, 2+1, 3+1]
-			VectorRegister LeftSamplesFrom, LeftSamplesTo;
+			VectorRegister4f LeftSamplesFrom, LeftSamplesTo;
 			VectorDeinterleave(LeftSamplesFrom, LeftSamplesTo, LeftSamples01, LeftSamples23);
 
 			// our lerp vector is CurFracVec, masking off the 1 bit in the mantissa, subtract 2.
 			CurFracVec = VectorBitwiseAnd(CurFracVec, MaskVec);
-			VectorRegister LerpFactor = VectorSubtract(CurFracVec, TwoVec);
-			VectorRegister OneMinusLerpFactor = VectorSubtract(ThreeVec, CurFracVec);
+			VectorRegister4f LerpFactor = VectorSubtract(CurFracVec, TwoVec);
+			VectorRegister4f OneMinusLerpFactor = VectorSubtract(ThreeVec, CurFracVec);
 
 			CurFracVec = VectorAdd(CurFracVec, StepVec);
-			VectorRegister OutputVec = VectorAdd(VectorMultiply(LeftSamplesFrom, OneMinusLerpFactor), VectorMultiply(LeftSamplesTo, LerpFactor));
+			VectorRegister4f OutputVec = VectorAdd(VectorMultiply(LeftSamplesFrom, OneMinusLerpFactor), VectorMultiply(LeftSamplesTo, LerpFactor));
 
 			VectorStore(OutputVec, OutputFrames);
 
@@ -138,25 +138,25 @@ uint32 Audio::FVectorLinearResampler::ResampleStereo(uint32 OutputFramesNeeded,
 	{
 		float CurrentFrameFractionFloat = (float)(CurrentFrameFraction & 0xffff);
 
-		VectorRegister ScaleVec = VectorSetFloat1(1.0f / 65536.0f);
-		VectorRegister TwoVec = VectorSetFloat1(2.0f);
-		VectorRegister OneVec = VectorSetFloat1(1.0f);
-		VectorRegister ThreeVec = VectorSetFloat1(3.0f);
+		VectorRegister4f ScaleVec = VectorSetFloat1(1.0f / 65536.0f);
+		VectorRegister4f TwoVec = VectorSetFloat1(2.0f);
+		VectorRegister4f OneVec = VectorSetFloat1(1.0f);
+		VectorRegister4f ThreeVec = VectorSetFloat1(3.0f);
 
-		VectorRegister CurFracStepVec = VectorSet(
+		VectorRegister4f CurFracStepVec = VectorSet(
 			(float)((0 * FixedPointSampleRate) & 0xffff),
 			(float)((1 * FixedPointSampleRate) & 0xffff),
 			(float)((2 * FixedPointSampleRate) & 0xffff),
 			(float)((3 * FixedPointSampleRate) & 0xffff));
 
-		VectorRegister CurFracVec = VectorSetFloat1(CurrentFrameFractionFloat);
+		VectorRegister4f CurFracVec = VectorSetFloat1(CurrentFrameFractionFloat);
 		CurFracVec = VectorAdd(CurFracVec, CurFracStepVec);
 		CurFracVec = VectorMultiply(CurFracVec, ScaleVec);
 		CurFracVec = VectorAdd(CurFracVec, TwoVec);
 
-		VectorRegister MaskVec = MakeVectorRegister(0xffbfffff, 0xffbfffff, 0xffbfffff, 0xffbfffff);
+		VectorRegister4f MaskVec = MakeVectorRegisterFloat(0xffbfffff, 0xffbfffff, 0xffbfffff, 0xffbfffff);
 
-		VectorRegister StepVec = VectorSetFloat1((float)((FixedPointSampleRate * 4) & 0xffff));
+		VectorRegister4f StepVec = VectorSetFloat1((float)((FixedPointSampleRate * 4) & 0xffff));
 		StepVec = VectorMultiply(StepVec, ScaleVec);
 
 		uint32 OutputFramesNeededSIMD = OutputFramesNeeded & ~3;
@@ -171,29 +171,29 @@ uint32 Audio::FVectorLinearResampler::ResampleStereo(uint32 OutputFramesNeeded,
 			};
 
 			// [0, 0+1, 1, 1+1]
-			VectorRegister LeftSamples01 = VectorLoadTwoPairsFloat(LeftSamples + SourceOffsets[0], LeftSamples + SourceOffsets[1]);
-			VectorRegister RightSamples01 = VectorLoadTwoPairsFloat(RightSamples + SourceOffsets[0], RightSamples + SourceOffsets[1]);
+			VectorRegister4f LeftSamples01 = VectorLoadTwoPairsFloat(LeftSamples + SourceOffsets[0], LeftSamples + SourceOffsets[1]);
+			VectorRegister4f RightSamples01 = VectorLoadTwoPairsFloat(RightSamples + SourceOffsets[0], RightSamples + SourceOffsets[1]);
 
 			// [2, 2+1, 3, 3+1]
-			VectorRegister LeftSamples23 = VectorLoadTwoPairsFloat(LeftSamples + SourceOffsets[2], LeftSamples + SourceOffsets[3]);
-			VectorRegister RightSamples23 = VectorLoadTwoPairsFloat(RightSamples + SourceOffsets[2], RightSamples + SourceOffsets[3]);
+			VectorRegister4f LeftSamples23 = VectorLoadTwoPairsFloat(LeftSamples + SourceOffsets[2], LeftSamples + SourceOffsets[3]);
+			VectorRegister4f RightSamples23 = VectorLoadTwoPairsFloat(RightSamples + SourceOffsets[2], RightSamples + SourceOffsets[3]);
 
 			// want [0, 1, 2, 3]
 			// [0+1, 1+1, 2+1, 3+1]
-			VectorRegister LeftSamplesFrom, LeftSamplesTo;
+			VectorRegister4f LeftSamplesFrom, LeftSamplesTo;
 			VectorDeinterleave(LeftSamplesFrom, LeftSamplesTo, LeftSamples01, LeftSamples23);
-			VectorRegister RightSamplesFrom, RightSamplesTo;
+			VectorRegister4f RightSamplesFrom, RightSamplesTo;
 			VectorDeinterleave(RightSamplesFrom, RightSamplesTo, RightSamples01, RightSamples23);
 
 			// our lerp vector is cur_frac_vec, masking off the 1 bit in the mantissa, subtract 2.
 			CurFracVec = VectorBitwiseAnd(CurFracVec, MaskVec);
-			VectorRegister LerpFactor = VectorSubtract(CurFracVec, TwoVec);
-			VectorRegister OneMinusLerpFactor = VectorSubtract(ThreeVec, CurFracVec);
+			VectorRegister4f LerpFactor = VectorSubtract(CurFracVec, TwoVec);
+			VectorRegister4f OneMinusLerpFactor = VectorSubtract(ThreeVec, CurFracVec);
 
 			CurFracVec = VectorAdd(CurFracVec, StepVec);
 
-			VectorRegister LeftOutputVec = VectorAdd(VectorMultiply(LeftSamplesFrom, OneMinusLerpFactor), VectorMultiply(LeftSamplesTo, LerpFactor));
-			VectorRegister RightOutputVec = VectorAdd(VectorMultiply(RightSamplesFrom, OneMinusLerpFactor), VectorMultiply(RightSamplesTo, LerpFactor));
+			VectorRegister4f LeftOutputVec = VectorAdd(VectorMultiply(LeftSamplesFrom, OneMinusLerpFactor), VectorMultiply(LeftSamplesTo, LerpFactor));
+			VectorRegister4f RightOutputVec = VectorAdd(VectorMultiply(RightSamplesFrom, OneMinusLerpFactor), VectorMultiply(RightSamplesTo, LerpFactor));
 
 			VectorStore(LeftOutputVec, OutputFrames);
 			VectorStore(RightOutputVec, OutputFrames + OutputFramesStrideFloats);
