Audio signal processing cleanup
- Removing need for alignment in audio DSP.
- Decreased header dependencies across codebase to improve build times
- Fixed improper `using namespace` issues.

diff --git a/Engine/Plugins/Experimental/MotoSynth/Source/MotoSynth/Private/MotoSynthEngine.cpp b/Engine/Plugins/Experimental/MotoSynth/Source/MotoSynth/Private/MotoSynthEngine.cpp
--- a/Engine/Plugins/Experimental/MotoSynth/Source/MotoSynth/Private/MotoSynthEngine.cpp
+++ b/Engine/Plugins/Experimental/MotoSynth/Source/MotoSynth/Private/MotoSynthEngine.cpp
@@ -3,6 +3,7 @@
 #include "MotoSynthEngine.h"
 #include "MotoSynthModule.h"
 #include "MotoSynthDataManager.h"
+#include "DSP/FloatArrayMath.h"
 
 static int32 MotosynthDisabledCVar = 0;
 FAutoConsoleVariableRef CVarDisableMotoSynth(
@@ -698,4 +699,4 @@ bool FMotoSynthGrainRuntime::IsDone() const
 void FMotoSynthGrainRuntime::SetRPM(int32 InRPM)
 {
 	CurrentRuntimeRPM = InRPM;
-}
\ No newline at end of file
+}
diff --git a/Engine/Plugins/Experimental/MotoSynth/Source/MotoSynth/Private/MotoSynthSourceAsset.cpp b/Engine/Plugins/Experimental/MotoSynth/Source/MotoSynth/Private/MotoSynthSourceAsset.cpp
--- a/Engine/Plugins/Experimental/MotoSynth/Source/MotoSynth/Private/MotoSynthSourceAsset.cpp
+++ b/Engine/Plugins/Experimental/MotoSynth/Source/MotoSynth/Private/MotoSynthSourceAsset.cpp
@@ -8,6 +8,7 @@
 #include "DSP/Filter.h"
 #include "DSP/DynamicsProcessor.h"
 #include "DSP/Dsp.h"
+#include "DSP/FloatArrayMath.h"
 #include "DSP/SampleRateConverter.h"
 #endif
 
diff --git a/Engine/Plugins/Runtime/AudioSynesthesia/Source/AudioSynesthesiaCore/Public/ConstantQNRTFactory.h b/Engine/Plugins/Runtime/AudioSynesthesia/Source/AudioSynesthesiaCore/Public/ConstantQNRTFactory.h
--- a/Engine/Plugins/Runtime/AudioSynesthesia/Source/AudioSynesthesiaCore/Public/ConstantQNRTFactory.h
+++ b/Engine/Plugins/Runtime/AudioSynesthesia/Source/AudioSynesthesiaCore/Public/ConstantQNRTFactory.h
@@ -3,9 +3,9 @@
 #pragma once
 
 #include "IAudioAnalyzerNRTInterface.h"
+#include "DSP/AlignedBuffer.h"
 #include "DSP/ConstantQ.h"
 #include "DSP/SlidingWindow.h"
-#include "DSP/FloatArrayMath.h"
 #include "ConstantQAnalyzer.h"
 
 namespace Audio
diff --git a/Engine/Plugins/Runtime/Metasound/Source/MetasoundFrontend/Public/MetasoundAudioBuffer.h b/Engine/Plugins/Runtime/Metasound/Source/MetasoundFrontend/Public/MetasoundAudioBuffer.h
--- a/Engine/Plugins/Runtime/Metasound/Source/MetasoundFrontend/Public/MetasoundAudioBuffer.h
+++ b/Engine/Plugins/Runtime/Metasound/Source/MetasoundFrontend/Public/MetasoundAudioBuffer.h
@@ -2,7 +2,7 @@
 
 #pragma once
 
-#include "DSP/FloatArrayMath.h"
+#include "DSP/AlignedBuffer.h"
 #include "MetasoundGraphCoreModule.h"
 #include "MetasoundDataReferenceMacro.h"
 #include "MetasoundLog.h"
@@ -106,6 +106,27 @@ return Buffer;
 				return Buffer;
 			}
 
+			/** Implicit conversion to Audio::FAlignedFloatBuffer */
+			FORCEINLINE operator TArrayView<const float> () const
+			{
+#if METASOUNDGRAPHCORE_CHECKAUDIONUM
+				UE_CLOG(InitialNum != Buffer.Num(), LogMetaSound, Error, TEXT("MetaSound audio buffer size change detected.  Audio Buffers should not be resized."));
+#endif
+return Buffer;
+			}
+
+			/** Implicit conversion to Audio::FAlignedFloatBuffer 
+			 *
+			 * WARNING: if the buffer is resized, it will cause errors.
+			 */
+			FORCEINLINE operator TArrayView<float> ()
+			{
+#if METASOUNDGRAPHCORE_CHECKAUDIONUM
+				UE_CLOG(InitialNum != Buffer.Num(), LogMetaSound, Error, TEXT("MetaSound audio buffer size change detected.  Audio Buffers should not be resized."));
+#endif				
+				return Buffer;
+			}
+
 			FORCEINLINE void Zero()
 			{
 				if (Buffer.Num() > 0)
diff --git a/Engine/Plugins/Runtime/Metasound/Source/MetasoundStandardNodes/Private/MetasoundAdditiveSynthNode.cpp b/Engine/Plugins/Runtime/Metasound/Source/MetasoundStandardNodes/Private/MetasoundAdditiveSynthNode.cpp
--- a/Engine/Plugins/Runtime/Metasound/Source/MetasoundStandardNodes/Private/MetasoundAdditiveSynthNode.cpp
+++ b/Engine/Plugins/Runtime/Metasound/Source/MetasoundStandardNodes/Private/MetasoundAdditiveSynthNode.cpp
@@ -1,6 +1,6 @@
 // Copyright Epic Games, Inc. All Rights Reserved.
 
-#include "DSP/BufferVectorOperations.h"
+#include "DSP/FloatArrayMath.h"
 #include "DSP/SineWaveTableOsc.h"
 #include "MetasoundFacade.h"
 #include "Internationalization/Text.h"
diff --git a/Engine/Plugins/Runtime/Oculus/OculusAudio/Source/OculusAudio/Private/OculusAmbisonicSpatializer.cpp b/Engine/Plugins/Runtime/Oculus/OculusAudio/Source/OculusAudio/Private/OculusAmbisonicSpatializer.cpp
--- a/Engine/Plugins/Runtime/Oculus/OculusAudio/Source/OculusAudio/Private/OculusAmbisonicSpatializer.cpp
+++ b/Engine/Plugins/Runtime/Oculus/OculusAudio/Source/OculusAudio/Private/OculusAmbisonicSpatializer.cpp
@@ -3,7 +3,7 @@
 #include "OculusAmbisonicSpatializer.h"
 #include "ISoundfieldFormat.h"
 #include "OculusAudioMixer.h"
-#include "DSP/BufferVectorOperations.h"
+#include "DSP/FloatArrayMath.h"
 #include "AudioMixerDevice.h"
 #include "OculusAudioDllManager.h"
 #include "OculusAudioContextManager.h"
diff --git a/Engine/Plugins/Runtime/SoundFields/Source/SoundFields/Private/SoundFields.cpp b/Engine/Plugins/Runtime/SoundFields/Source/SoundFields/Private/SoundFields.cpp
--- a/Engine/Plugins/Runtime/SoundFields/Source/SoundFields/Private/SoundFields.cpp
+++ b/Engine/Plugins/Runtime/SoundFields/Source/SoundFields/Private/SoundFields.cpp
@@ -2,8 +2,9 @@
 
 #include "SoundFields.h"
 #include "CoreMinimal.h"
+#include "DSP/AlignedBuffer.h"
 #include "DSP/Dsp.h"
-#include "DSP/BufferVectorOperations.h"
+#include "DSP/FloatArrayMath.h"
 #include "SoundFieldRendering.h"
 #include "HAL/IConsoleManager.h"
 
diff --git a/Engine/Plugins/Runtime/Synthesis/Source/Synthesis/Private/ConvolutionReverb.cpp b/Engine/Plugins/Runtime/Synthesis/Source/Synthesis/Private/ConvolutionReverb.cpp
--- a/Engine/Plugins/Runtime/Synthesis/Source/Synthesis/Private/ConvolutionReverb.cpp
+++ b/Engine/Plugins/Runtime/Synthesis/Source/Synthesis/Private/ConvolutionReverb.cpp
@@ -4,7 +4,7 @@
 
 #include "AudioMixerDevice.h"
 #include "CoreMinimal.h"
-#include "DSP/BufferVectorOperations.h"
+#include "DSP/FloatArrayMath.h"
 #include "DSP/ConvolutionAlgorithm.h"
 #include "DSP/SampleRateConverter.h"
 #include "SubmixChannelFormatConverter.h"
diff --git a/Engine/Plugins/Runtime/Synthesis/Source/Synthesis/Private/SourceEffectConvolutionReverb.cpp b/Engine/Plugins/Runtime/Synthesis/Source/Synthesis/Private/SourceEffectConvolutionReverb.cpp
--- a/Engine/Plugins/Runtime/Synthesis/Source/Synthesis/Private/SourceEffectConvolutionReverb.cpp
+++ b/Engine/Plugins/Runtime/Synthesis/Source/Synthesis/Private/SourceEffectConvolutionReverb.cpp
@@ -2,8 +2,8 @@
 
 #include "SourceEffects/SourceEffectConvolutionReverb.h"
 
-#include "CoreMinimal.h"
 #include "Async/Async.h"
+#include "DSP/FloatArrayMath.h"
 #include "SynthesisModule.h"
 
 FSourceEffectConvolutionReverbSettings::FSourceEffectConvolutionReverbSettings()
diff --git a/Engine/Plugins/Runtime/Synthesis/Source/Synthesis/Private/SourceEffectFilter.cpp b/Engine/Plugins/Runtime/Synthesis/Source/Synthesis/Private/SourceEffectFilter.cpp
--- a/Engine/Plugins/Runtime/Synthesis/Source/Synthesis/Private/SourceEffectFilter.cpp
+++ b/Engine/Plugins/Runtime/Synthesis/Source/Synthesis/Private/SourceEffectFilter.cpp
@@ -1,7 +1,7 @@
 // Copyright Epic Games, Inc. All Rights Reserved.
 
 #include "SourceEffects/SourceEffectFilter.h"
-#include "DSP/BufferVectorOperations.h"
+#include "DSP/FloatArrayMath.h"
 #include "DSP/Dsp.h"
 #include "AudioMixerDevice.h"
 
diff --git a/Engine/Plugins/Runtime/Synthesis/Source/Synthesis/Private/SourceEffectMotionFilter.cpp b/Engine/Plugins/Runtime/Synthesis/Source/Synthesis/Private/SourceEffectMotionFilter.cpp
--- a/Engine/Plugins/Runtime/Synthesis/Source/Synthesis/Private/SourceEffectMotionFilter.cpp
+++ b/Engine/Plugins/Runtime/Synthesis/Source/Synthesis/Private/SourceEffectMotionFilter.cpp
@@ -1,8 +1,9 @@
 // Copyright Epic Games, Inc. All Rights Reserved.
 
-
 #include "SourceEffects/SourceEffectMotionFilter.h"
+
 #include "AudioMixer.h"
+#include "DSP/FloatArrayMath.h"
 
 FSourceEffectMotionFilter::FSourceEffectMotionFilter()
 	: Topology(ESourceEffectMotionFilterTopology::ParallelMode)
@@ -475,4 +476,4 @@ void USourceEffectMotionFilterPreset::SetSettings(const FSourceEffectMotionFilte
 {
 	// Performs necessary broadcast to effect instances
 	UpdateSettings(InSettings);
-}
\ No newline at end of file
+}
diff --git a/Engine/Plugins/Runtime/Synthesis/Source/Synthesis/Private/SubmixChannelFormatConverter.cpp b/Engine/Plugins/Runtime/Synthesis/Source/Synthesis/Private/SubmixChannelFormatConverter.cpp
--- a/Engine/Plugins/Runtime/Synthesis/Source/Synthesis/Private/SubmixChannelFormatConverter.cpp
+++ b/Engine/Plugins/Runtime/Synthesis/Source/Synthesis/Private/SubmixChannelFormatConverter.cpp
@@ -4,6 +4,7 @@
 
 #include "AudioMixerDevice.h"
 #include "DSP/AudioChannelFormatConverter.h"
+#include "DSP/FloatArrayMath.h"
 #include "SynthesisModule.h"
 
 namespace Audio
diff --git a/Engine/Plugins/Runtime/Synthesis/Source/Synthesis/Private/SubmixEffectConvolutionReverb.cpp b/Engine/Plugins/Runtime/Synthesis/Source/Synthesis/Private/SubmixEffectConvolutionReverb.cpp
--- a/Engine/Plugins/Runtime/Synthesis/Source/Synthesis/Private/SubmixEffectConvolutionReverb.cpp
+++ b/Engine/Plugins/Runtime/Synthesis/Source/Synthesis/Private/SubmixEffectConvolutionReverb.cpp
@@ -2,8 +2,8 @@
 
 #include "SubmixEffects/SubmixEffectConvolutionReverb.h"
 
-#include "CoreMinimal.h"
 #include "Async/Async.h"
+#include "DSP/FloatArrayMath.h"
 #include "SynthesisModule.h"
 
 FSubmixEffectConvolutionReverbSettings::FSubmixEffectConvolutionReverbSettings()
diff --git a/Engine/Plugins/Runtime/Synthesis/Source/Synthesis/Private/SubmixEffectMultiBandCompressor.cpp b/Engine/Plugins/Runtime/Synthesis/Source/Synthesis/Private/SubmixEffectMultiBandCompressor.cpp
--- a/Engine/Plugins/Runtime/Synthesis/Source/Synthesis/Private/SubmixEffectMultiBandCompressor.cpp
+++ b/Engine/Plugins/Runtime/Synthesis/Source/Synthesis/Private/SubmixEffectMultiBandCompressor.cpp
@@ -5,7 +5,7 @@
 #include "AudioDeviceManager.h"
 #include "AudioMixerDevice.h"
 #include "AudioMixerSubmix.h"
-#include "DSP/BufferVectorOperations.h"
+#include "DSP/FloatArrayMath.h"
 
 FSubmixEffectMultibandCompressor::FSubmixEffectMultibandCompressor()
 {
diff --git a/Engine/Plugins/Runtime/Synthesis/Source/Synthesis/Private/SynthComponentMonoWaveTable.cpp b/Engine/Plugins/Runtime/Synthesis/Source/Synthesis/Private/SynthComponentMonoWaveTable.cpp
--- a/Engine/Plugins/Runtime/Synthesis/Source/Synthesis/Private/SynthComponentMonoWaveTable.cpp
+++ b/Engine/Plugins/Runtime/Synthesis/Source/Synthesis/Private/SynthComponentMonoWaveTable.cpp
@@ -2,7 +2,7 @@
 
 
 #include "SynthComponents/SynthComponentMonoWaveTable.h"
-#include "DSP/BufferVectorOperations.h"
+#include "DSP/FloatArrayMath.h"
 
 UMonoWaveTableSynthPreset::UMonoWaveTableSynthPreset() : PresetName(TEXT("Default"))
 , LockKeyframesToGrid(16)
diff --git a/Engine/Source/Runtime/AudioMixer/Private/AudioMixerBus.cpp b/Engine/Source/Runtime/AudioMixer/Private/AudioMixerBus.cpp
--- a/Engine/Source/Runtime/AudioMixer/Private/AudioMixerBus.cpp
+++ b/Engine/Source/Runtime/AudioMixer/Private/AudioMixerBus.cpp
@@ -4,7 +4,8 @@
 
 #include "Algo/ForEach.h"
 #include "AudioMixerSourceManager.h"
-#include "DSP/BufferVectorOperations.h"
+#include "DSP/AlignedBuffer.h"
+#include "DSP/FloatArrayMath.h"
 
 namespace Audio
 {
diff --git a/Engine/Source/Runtime/AudioMixer/Private/AudioMixerSourceManager.cpp b/Engine/Source/Runtime/AudioMixer/Private/AudioMixerSourceManager.cpp
--- a/Engine/Source/Runtime/AudioMixer/Private/AudioMixerSourceManager.cpp
+++ b/Engine/Source/Runtime/AudioMixer/Private/AudioMixerSourceManager.cpp
@@ -1,12 +1,14 @@
 // Copyright Epic Games, Inc. All Rights Reserved.
 
 #include "AudioMixerSourceManager.h"
+
 #include "AudioMixerSourceBuffer.h"
 #include "AudioMixerSource.h"
 #include "AudioMixerDevice.h"
 #include "AudioMixerSourceVoice.h"
 #include "AudioMixerSubmix.h"
 #include "AudioThread.h"
+#include "DSP/FloatArrayMath.h"
 #include "IAudioExtensionPlugin.h"
 #include "AudioMixer.h"
 #include "Sound/SoundModulationDestination.h"
diff --git a/Engine/Source/Runtime/AudioMixer/Private/AudioMixerSourceOutputBuffer.cpp b/Engine/Source/Runtime/AudioMixer/Private/AudioMixerSourceOutputBuffer.cpp
--- a/Engine/Source/Runtime/AudioMixer/Private/AudioMixerSourceOutputBuffer.cpp
+++ b/Engine/Source/Runtime/AudioMixer/Private/AudioMixerSourceOutputBuffer.cpp
@@ -3,7 +3,7 @@
 #include "AudioMixerSourceOutputBuffer.h"
 #include "SoundFieldRendering.h"
 #include "AudioMixerSubmix.h"
-
+#include "DSP/FloatArrayMath.h"
 
 namespace Audio
 {
@@ -534,4 +534,4 @@ namespace Audio
 	}
 
 
-}
\ No newline at end of file
+}
diff --git a/Engine/Source/Runtime/AudioMixerCore/Private/AudioMixer.cpp b/Engine/Source/Runtime/AudioMixerCore/Private/AudioMixer.cpp
--- a/Engine/Source/Runtime/AudioMixerCore/Private/AudioMixer.cpp
+++ b/Engine/Source/Runtime/AudioMixerCore/Private/AudioMixer.cpp
@@ -2,6 +2,7 @@
 
 #include "AudioMixer.h"
 #include "DSP/BufferVectorOperations.h"
+#include "DSP/FloatArrayMath.h"
 #include "HAL/RunnableThread.h"
 #include "HAL/ThreadSafeCounter.h"
 #include "HAL/IConsoleManager.h"
diff --git a/Engine/Source/Runtime/Engine/Private/SoundModulationDestination.cpp b/Engine/Source/Runtime/Engine/Private/SoundModulationDestination.cpp
--- a/Engine/Source/Runtime/Engine/Private/SoundModulationDestination.cpp
+++ b/Engine/Source/Runtime/Engine/Private/SoundModulationDestination.cpp
@@ -4,6 +4,7 @@
 
 #include "Async/Async.h"
 #include "AudioDevice.h"
+#include "DSP/FloatArrayMath.h"
 #include "Math/TransformCalculus.h"
 #include "UObject/Object.h"
 
diff --git a/Engine/Source/Runtime/Engine/Public/AudioDevice.h b/Engine/Source/Runtime/Engine/Public/AudioDevice.h
--- a/Engine/Source/Runtime/Engine/Public/AudioDevice.h
+++ b/Engine/Source/Runtime/Engine/Public/AudioDevice.h
@@ -6,7 +6,6 @@
 #include "Components/AudioComponent.h"
 #include "CoreMinimal.h"
 #include "DSP/MultithreadedPatching.h"
-#include "DSP/SpectrumAnalyzer.h"
 #include "Engine/Engine.h"
 #include "EngineGlobals.h"
 #include "HAL/LowLevelMemStats.h"
@@ -28,6 +27,11 @@
 #include "AudioMixer.h"
 #include "UObject/StrongObjectPtr.h"
 
+#if UE_ENABLE_INCLUDE_ORDER_DEPRECATED_IN_5_1
+#include "DSP/SpectrumAnalyzer.h"
+#endif // UE_ENABLE_INCLUDE_ORDER_DEPRECATED_IN_5_1
+
+
 /**
  * Forward declares
  */
diff --git a/Engine/Source/Runtime/SignalProcessing/Private/AlignedBlockBuffer.cpp b/Engine/Source/Runtime/SignalProcessing/Private/AlignedBlockBuffer.cpp
--- a/Engine/Source/Runtime/SignalProcessing/Private/AlignedBlockBuffer.cpp
+++ b/Engine/Source/Runtime/SignalProcessing/Private/AlignedBlockBuffer.cpp
@@ -4,324 +4,324 @@
 #include "CoreMinimal.h"
 #include "SignalProcessingModule.h"
 
-using namespace Audio;
-
-FAlignedBlockBuffer::FAlignedBlockBuffer(int32 InSampleCapacity, int32 InMaxNumInspectSamples, uint32 InByteAlignment, uint32 InAllocByteAlignment)
-	:	AllocByteAlignment(InAllocByteAlignment)
-	,	ByteAlignment(InByteAlignment)
-	,	FloatAlignment(0)
-	,	InternalBuffer(nullptr)
-	,	RolloverBuffer(nullptr)
-	,	ReadPointer(nullptr)
-	,	WritePointer(nullptr)
-	,	SampleCapacity(InSampleCapacity)
-	,	MaxNumInspectSamples(InMaxNumInspectSamples)
-	,	WriteCount(0)
-	,	WritePos(0)
-	,	ReadPos(0)
+namespace Audio
 {
-	FloatAlignment = ByteAlignment / sizeof(float);
-
-	int32 MaxStorage = TNumericLimits<int32>::Max() / 4;
-
-	// The SampleCapacity needs to also be divisible by the byte alignment as it is
-	// used in a bunch of the pointer math.  This ensures that the ReadPointer
-	// is always pointing to aligned memory
-	while (0 != (SampleCapacity % FloatAlignment))
+	FAlignedBlockBuffer::FAlignedBlockBuffer(int32 InSampleCapacity, int32 InMaxNumInspectSamples, uint32 InByteAlignment, uint32 InAllocByteAlignment)
+		:	AllocByteAlignment(InAllocByteAlignment)
+		,	ByteAlignment(InByteAlignment)
+		,	FloatAlignment(0)
+		,	InternalBuffer(nullptr)
+		,	RolloverBuffer(nullptr)
+		,	ReadPointer(nullptr)
+		,	WritePointer(nullptr)
+		,	SampleCapacity(InSampleCapacity)
+		,	MaxNumInspectSamples(InMaxNumInspectSamples)
+		,	WriteCount(0)
+		,	WritePos(0)
+		,	ReadPos(0)
 	{
-		SampleCapacity++;
+		FloatAlignment = ByteAlignment / sizeof(float);
+
+		int32 MaxStorage = TNumericLimits<int32>::Max() / 4;
+
+		// The SampleCapacity needs to also be divisible by the byte alignment as it is
+		// used in a bunch of the pointer math.  This ensures that the ReadPointer
+		// is always pointing to aligned memory
+		while (0 != (SampleCapacity % FloatAlignment))
+		{
+			SampleCapacity++;
+		}
+
+		// Check sanity of variables.
+		if (MaxNumInspectSamples > MaxStorage)
+		{
+			UE_LOG(LogSignalProcessing, Warning, TEXT("MaxNumInspectSamples [%d] reduced to maximum storage [%d]"), MaxNumInspectSamples, MaxStorage);
+			MaxNumInspectSamples = MaxStorage;
+		}
+
+		if (SampleCapacity > MaxStorage)
+		{
+			UE_LOG(LogSignalProcessing, Warning, TEXT("SampleCapacity [%d] reduced to maximum storage [%d]"), SampleCapacity, MaxStorage);
+			SampleCapacity = MaxStorage;
+		}
+
+		if (MaxNumInspectSamples > SampleCapacity)
+		{
+			UE_LOG(LogSignalProcessing, Warning, TEXT("SampleCapacity [%d] increased to Maxmimum Inspection [%d]"), SampleCapacity, MaxNumInspectSamples);
+			SampleCapacity = MaxNumInspectSamples;
+		}
+
+		// Allocate space
+		RolloverBuffer = AllocateAlignedFloatArray(MaxNumInspectSamples);
+		InternalBuffer = AllocateAlignedFloatArray(SampleCapacity);
+		
+		// Set everything to zero
+		ZeroFloatArray(RolloverBuffer, MaxNumInspectSamples);
+		ZeroFloatArray(InternalBuffer, SampleCapacity);
+		
+		// Initialize pointers to internal buffers
+		WritePointer = InternalBuffer;
+		ReadPointer = InternalBuffer;
 	}
 
-	// Check sanity of variables.
-	if (MaxNumInspectSamples > MaxStorage)
+			
+	FAlignedBlockBuffer::~FAlignedBlockBuffer() throw()
 	{
-		UE_LOG(LogSignalProcessing, Warning, TEXT("MaxNumInspectSamples [%d] reduced to maximum storage [%d]"), MaxNumInspectSamples, MaxStorage);
-		MaxNumInspectSamples = MaxStorage;
+		FreeFloatArray(InternalBuffer);
+		FreeFloatArray(RolloverBuffer);
 	}
 
-	if (SampleCapacity > MaxStorage)
+	int32 FAlignedBlockBuffer::GetSampleCapacity( ) const
 	{
-		UE_LOG(LogSignalProcessing, Warning, TEXT("SampleCapacity [%d] reduced to maximum storage [%d]"), SampleCapacity, MaxStorage);
-		SampleCapacity = MaxStorage;
+		return SampleCapacity;
 	}
-
-	if (MaxNumInspectSamples > SampleCapacity)
+			
+	int32 FAlignedBlockBuffer::GetNumAvailable() const
 	{
-		UE_LOG(LogSignalProcessing, Warning, TEXT("SampleCapacity [%d] increased to Maxmimum Inspection [%d]"), SampleCapacity, MaxNumInspectSamples);
-		SampleCapacity = MaxNumInspectSamples;
+		return WriteCount;
 	}
 
-	// Allocate space
-	RolloverBuffer = AllocateAlignedFloatArray(MaxNumInspectSamples);
-	InternalBuffer = AllocateAlignedFloatArray(SampleCapacity);
-	
-	// Set everything to zero
-	ZeroFloatArray(RolloverBuffer, MaxNumInspectSamples);
-	ZeroFloatArray(InternalBuffer, SampleCapacity);
-	
-	// Initialize pointers to internal buffers
-	WritePointer = InternalBuffer;
-	ReadPointer = InternalBuffer;
-}
-
-		
-FAlignedBlockBuffer::~FAlignedBlockBuffer() throw()
-{
-	FreeFloatArray(InternalBuffer);
-	FreeFloatArray(RolloverBuffer);
-}
-
-int32 FAlignedBlockBuffer::GetSampleCapacity( ) const
-{
-	return SampleCapacity;
-}
-		
-int32 FAlignedBlockBuffer::GetNumAvailable() const
-{
-	return WriteCount;
-}
-
-int32 FAlignedBlockBuffer::CheckNumAndIncrementWriteCount(int32 InNum)
-{
-	// Clamp InNum to capacity
-	checkf(InNum <= SampleCapacity, TEXT("Input buffer [%d] larger than capacity [%d]"), InNum, SampleCapacity);
-	InNum = InNum < SampleCapacity ? InNum : SampleCapacity;
-
-	// Increment write count keeping within capacity
-	WriteCount += InNum;
-
-	// Handle buffer overflow
-	checkf(WriteCount <= SampleCapacity, TEXT("Buffer overflow"));
-	if(WriteCount > SampleCapacity)
+	int32 FAlignedBlockBuffer::CheckNumAndIncrementWriteCount(int32 InNum)
 	{
-		// Reduce number of input samples on buffer overflow.
-		InNum -= WriteCount - SampleCapacity;
-		WriteCount = SampleCapacity;
+		// Clamp InNum to capacity
+		checkf(InNum <= SampleCapacity, TEXT("Input buffer [%d] larger than capacity [%d]"), InNum, SampleCapacity);
+		InNum = InNum < SampleCapacity ? InNum : SampleCapacity;
+
+		// Increment write count keeping within capacity
+		WriteCount += InNum;
+
+		// Handle buffer overflow
+		checkf(WriteCount <= SampleCapacity, TEXT("Buffer overflow"));
+		if(WriteCount > SampleCapacity)
+		{
+			// Reduce number of input samples on buffer overflow.
+			InNum -= WriteCount - SampleCapacity;
+			WriteCount = SampleCapacity;
+		}
+
+		return InNum;
 	}
 
-	return InNum;
-}
-
-void FAlignedBlockBuffer::AddZeros(int32 InNum)
-{
-	InNum = CheckNumAndIncrementWriteCount(InNum);
-	
-	// Check if we expect to be wrapping around end of input buffer
-	int32 RolloverState = SampleCapacity - (WritePos + InNum);
-	
-	if(RolloverState > 0)
+	void FAlignedBlockBuffer::AddZeros(int32 InNum)
 	{
-		// No rollover if rollover state is above zero.
-		ZeroFloatArray(WritePointer, InNum);
-		WritePos += InNum;
-		WritePointer += InNum;
-	}
-	else
-	{
-		// Zeroing buffer goes over internal buffer border, so perform zeroing
-		// in two parts.
+		InNum = CheckNumAndIncrementWriteCount(InNum);
 		
-		// Zero first part of buffer
-		int32 RemainingBufferSamples = InNum + RolloverState;
-		ZeroFloatArray(WritePointer, RemainingBufferSamples);
-		WritePointer = InternalBuffer;
-		WritePos = -RolloverState;
+		// Check if we expect to be wrapping around end of input buffer
+		int32 RolloverState = SampleCapacity - (WritePos + InNum);
 		
-		// Zero the second part
-		ZeroFloatArray(WritePointer, -RolloverState);
-		WritePointer -= RolloverState;
-	}
-}
-		
-void FAlignedBlockBuffer::AddSamples(const float* InSamples, int32 InNum)
-{
-	checkf(nullptr != InSamples, TEXT("Null pointer"));
-	if (nullptr == InSamples)
-	{
-		// Quit early with null pointers.
-		return;
+		if(RolloverState > 0)
+		{
+			// No rollover if rollover state is above zero.
+			ZeroFloatArray(WritePointer, InNum);
+			WritePos += InNum;
+			WritePointer += InNum;
+		}
+		else
+		{
+			// Zeroing buffer goes over internal buffer border, so perform zeroing
+			// in two parts.
+			
+			// Zero first part of buffer
+			int32 RemainingBufferSamples = InNum + RolloverState;
+			ZeroFloatArray(WritePointer, RemainingBufferSamples);
+			WritePointer = InternalBuffer;
+			WritePos = -RolloverState;
+			
+			// Zero the second part
+			ZeroFloatArray(WritePointer, -RolloverState);
+			WritePointer -= RolloverState;
+		}
 	}
-
-	InNum = CheckNumAndIncrementWriteCount(InNum);
-	
-	// If the write operation is expected to go past the internal buffer boundary
-	// the RolloverState will be negative. If RolloverState is negative, the copy
-	// needs to happen in two parts. 
-	int32 RolloverState = SampleCapacity - (WritePos + InNum);
-	
-	if(RolloverState > 0)
+			
+	void FAlignedBlockBuffer::AddSamples(const float* InSamples, int32 InNum)
 	{
-		// No rollover if rollover state above zero.
-		CopyFloatArray(WritePointer, InSamples, InNum);
+		checkf(nullptr != InSamples, TEXT("Null pointer"));
+		if (nullptr == InSamples)
+		{
+			// Quit early with null pointers.
+			return;
+		}
+
+		InNum = CheckNumAndIncrementWriteCount(InNum);
 		
-		WritePos += InNum;
-		WritePointer += InNum;
-	}
-	else
-	{
-		//Copy first part of buffer
-		int32 RemainingBufferSamples = InNum + RolloverState;
-	
-		CopyFloatArray(WritePointer, InSamples, RemainingBufferSamples);
+		// If the write operation is expected to go past the internal buffer boundary
+		// the RolloverState will be negative. If RolloverState is negative, the copy
+		// needs to happen in two parts. 
+		int32 RolloverState = SampleCapacity - (WritePos + InNum);
 		
-		WritePointer = InternalBuffer;
-		WritePos = -RolloverState;
-		InSamples += RemainingBufferSamples;
+		if(RolloverState > 0)
+		{
+			// No rollover if rollover state above zero.
+			CopyFloatArray(WritePointer, InSamples, InNum);
+			
+			WritePos += InNum;
+			WritePointer += InNum;
+		}
+		else
+		{
+			//Copy first part of buffer
+			int32 RemainingBufferSamples = InNum + RolloverState;
 		
-		//Copy the second part
-		CopyFloatArray(WritePointer, InSamples, -RolloverState);
-
-		WritePointer -= RolloverState;
-	}
-}
-
-void FAlignedBlockBuffer::RemoveSamples(int32 InNum)
-{
-	// In order to keep the read pointer aligned, any operations that move the read pointer
-	// must also be aligned. 
-	checkf(0 == (InNum % FloatAlignment), TEXT("InNum must be divisible by alignement."));
-
-	//decrement the available samples
-	WriteCount -= InNum;
-
-	checkf(WriteCount >= 0, TEXT("Buffer underflow"));
-	if (WriteCount < 0)
-	{
-		// Gracefully handle underflow.
-		WriteCount = 0;
-		ReadPos = 0;
-		ReadPointer = InternalBuffer;
-		return;
+			CopyFloatArray(WritePointer, InSamples, RemainingBufferSamples);
+			
+			WritePointer = InternalBuffer;
+			WritePos = -RolloverState;
+			InSamples += RemainingBufferSamples;
+			
+			//Copy the second part
+			CopyFloatArray(WritePointer, InSamples, -RolloverState);
+
+			WritePointer -= RolloverState;
+		}
 	}
 
-	int32 RolloverState = SampleCapacity - (ReadPos + InNum);
-	
-	if( RolloverState > 0 )
+	void FAlignedBlockBuffer::RemoveSamples(int32 InNum)
 	{
-		ReadPos += InNum;
-		ReadPointer += InNum;
+		// In order to keep the read pointer aligned, any operations that move the read pointer
+		// must also be aligned. 
+		checkf(0 == (InNum % FloatAlignment), TEXT("InNum must be divisible by alignement."));
+
+		//decrement the available samples
+		WriteCount -= InNum;
+
+		checkf(WriteCount >= 0, TEXT("Buffer underflow"));
+		if (WriteCount < 0)
+		{
+			// Gracefully handle underflow.
+			WriteCount = 0;
+			ReadPos = 0;
+			ReadPointer = InternalBuffer;
+			return;
+		}
+
+		int32 RolloverState = SampleCapacity - (ReadPos + InNum);
+		
+		if( RolloverState > 0 )
+		{
+			ReadPos += InNum;
+			ReadPointer += InNum;
+		}
+		else
+		{
+			// Wrap the get pointer around in case where buffer boundary is crossed.
+			ReadPos = -RolloverState;
+			ReadPointer = InternalBuffer - RolloverState;
+		}
 	}
-	else
+			
+	const float* FAlignedBlockBuffer::InspectSamples(int32 InNum, int32 InOffset)
 	{
-		// Wrap the get pointer around in case where buffer boundary is crossed.
-		ReadPos = -RolloverState;
-		ReadPointer = InternalBuffer - RolloverState;
-	}
-}
+		checkf(0 == (InOffset % FloatAlignment), TEXT("InOffset must divisible by alignment"));
+		checkf(InNum <= MaxNumInspectSamples, TEXT("Requested samples [%d] are more than maximum [%d]"), InNum, MaxNumInspectSamples);
+		checkf((InNum + InOffset) <= WriteCount, TEXT("Buffer underflow"));
+		checkf(InOffset >= 0, TEXT("InOffset must be greater than 0"));
+
+		if (InNum < 1)
+		{
+			return nullptr;
+		}
+
+		if (InNum > MaxNumInspectSamples)
+		{
+			// This is an unfortunate situation because it allocates memory.
+			SetMaxNumInspectSamples(InNum);
+		}
+
+		if ((InNum + InOffset) > WriteCount)
+		{
+			// This is a pretty bad situation. Anything done here is a bad option.
+			// Hopefully the lesser of all the evils is returning a nullptr so that the issue
+			// can be identified quickly. 
+			return nullptr;
+		}
+
+		if (InOffset < 0)
+		{
+			// Invalid offset. Despite the profound implications for the advancement of humankind, we cannot see into the future. 
+			InOffset = 0;
+		}
 		
-const float* FAlignedBlockBuffer::InspectSamples(int32 InNum, int32 InOffset)
-{
-	checkf(0 == (InOffset % FloatAlignment), TEXT("InOffset must divisible by alignment"));
-	checkf(InNum <= MaxNumInspectSamples, TEXT("Requested samples [%d] are more than maximum [%d]"), InNum, MaxNumInspectSamples);
-	checkf((InNum + InOffset) <= WriteCount, TEXT("Buffer underflow"));
-	checkf(InOffset >= 0, TEXT("InOffset must be greater than 0"));
-
-	if (InNum < 1)
-	{
-		return nullptr;
+		int32 StartPos = ReadPos + InOffset;
+		int32 EndPos = StartPos + InNum;
+
+		while (StartPos > SampleCapacity)
+		{
+			StartPos -= SampleCapacity;
+			EndPos -= SampleCapacity;
+		}
+
+		if (EndPos < SampleCapacity)
+		{
+			return &InternalBuffer[StartPos];
+		}
+		else
+		{
+			int32 NumFirstHalf = SampleCapacity - StartPos;
+			// Copy first half into rolloverbuffer
+			CopyFloatArray(RolloverBuffer, &InternalBuffer[StartPos], NumFirstHalf);
+			// Copy second half into rolloverbuffer
+			CopyFloatArray(&RolloverBuffer[NumFirstHalf], InternalBuffer, InNum - NumFirstHalf);
+
+			return RolloverBuffer;
+		}
 	}
-
-	if (InNum > MaxNumInspectSamples)
+			
+	void FAlignedBlockBuffer::ClearSamples()
 	{
-		// This is an unfortunate situation because it allocates memory.
-		SetMaxNumInspectSamples(InNum);
-	}
+		ZeroFloatArray(InternalBuffer, SampleCapacity);
+		ZeroFloatArray(RolloverBuffer, MaxNumInspectSamples);
 
-	if ((InNum + InOffset) > WriteCount)
-	{
-		// This is a pretty bad situation. Anything done here is a bad option.
-		// Hopefully the lesser of all the evils is returning a nullptr so that the issue
-		// can be identified quickly. 
-		return nullptr;
+		ReadPointer = InternalBuffer;
+		WritePointer = InternalBuffer;
+		ReadPos = 0;
+		WritePos = 0;
+		
+		WriteCount = 0;
 	}
-
-	if (InOffset < 0)
+					
+	int32 FAlignedBlockBuffer::GetMaxNumInspectSamples() const
 	{
-		// Invalid offset. Despite the profound implications for the advancement of humankind, we cannot see into the future. 
-		InOffset = 0;
+		return MaxNumInspectSamples;
 	}
-	
-	int32 StartPos = ReadPos + InOffset;
-	int32 EndPos = StartPos + InNum;
 
-	while (StartPos > SampleCapacity)
+	void FAlignedBlockBuffer::SetMaxNumInspectSamples(int32 InMax)
 	{
-		StartPos -= SampleCapacity;
-		EndPos -= SampleCapacity;
+		if(InMax > MaxNumInspectSamples)
+		{
+			FreeFloatArray(RolloverBuffer);
+			RolloverBuffer = AllocateAlignedFloatArray(InMax);
+		}
+		
+		MaxNumInspectSamples = InMax;
 	}
 
-	if (EndPos < SampleCapacity)
+	float* FAlignedBlockBuffer::AllocateAlignedFloatArray(int32 InNum) const 
 	{
-		return &InternalBuffer[StartPos];
-	}
-	else
-	{
-		int32 NumFirstHalf = SampleCapacity - StartPos;
-		// Copy first half into rolloverbuffer
-		CopyFloatArray(RolloverBuffer, &InternalBuffer[StartPos], NumFirstHalf);
-		// Copy second half into rolloverbuffer
-		CopyFloatArray(&RolloverBuffer[NumFirstHalf], InternalBuffer, InNum - NumFirstHalf);
-
-		return RolloverBuffer;
+		int32 Size = InNum * sizeof(float);
+		return (float*)FMemory::Malloc(Size, AllocByteAlignment);
 	}
-}
-		
-void FAlignedBlockBuffer::ClearSamples()
-{
-	ZeroFloatArray(InternalBuffer, SampleCapacity);
-	ZeroFloatArray(RolloverBuffer, MaxNumInspectSamples);
-
-	ReadPointer = InternalBuffer;
-	WritePointer = InternalBuffer;
-	ReadPos = 0;
-	WritePos = 0;
-	
-	WriteCount = 0;
-}
-				
-int32 FAlignedBlockBuffer::GetMaxNumInspectSamples() const
-{
-	return MaxNumInspectSamples;
-}
 
-void FAlignedBlockBuffer::SetMaxNumInspectSamples(int32 InMax)
-{
-	if(InMax > MaxNumInspectSamples)
+	void FAlignedBlockBuffer::FreeFloatArray(float*& Array) const
 	{
-		FreeFloatArray(RolloverBuffer);
-		RolloverBuffer = AllocateAlignedFloatArray(InMax);
+		FMemory::Free(Array);
+		Array = nullptr;
 	}
-	
-	MaxNumInspectSamples = InMax;
-}
-
-float* FAlignedBlockBuffer::AllocateAlignedFloatArray(int32 InNum) const 
-{
-	int32 Size = InNum * sizeof(float);
-	return (float*)FMemory::Malloc(Size, AllocByteAlignment);
-}
-
-void FAlignedBlockBuffer::FreeFloatArray(float*& Array) const
-{
-	FMemory::Free(Array);
-	Array = nullptr;
-}
 
-void FAlignedBlockBuffer::ZeroFloatArray(float* InArray, int32 InNum) const
-{
-	if (InNum > 0) 
+	void FAlignedBlockBuffer::ZeroFloatArray(float* InArray, int32 InNum) const
 	{
-		int32 Size = InNum * sizeof(float);
-		FMemory::Memset(InArray, 0, Size);
+		if (InNum > 0) 
+		{
+			int32 Size = InNum * sizeof(float);
+			FMemory::Memset(InArray, 0, Size);
+		}
 	}
-}
 
-void FAlignedBlockBuffer::CopyFloatArray(float* ToArray, const float* FromArray, int32 InNum) const 
-{
-	if (InNum > 0)
+	void FAlignedBlockBuffer::CopyFloatArray(float* ToArray, const float* FromArray, int32 InNum) const 
 	{
-		int32 Size = InNum * sizeof(float);
-		FMemory::Memcpy(ToArray, FromArray, Size);
+		if (InNum > 0)
+		{
+			int32 Size = InNum * sizeof(float);
+			FMemory::Memcpy(ToArray, FromArray, Size);
+		}
 	}
 }
-
diff --git a/Engine/Source/Runtime/SignalProcessing/Private/AllPassFractionalDelay.cpp b/Engine/Source/Runtime/SignalProcessing/Private/AllPassFractionalDelay.cpp
--- a/Engine/Source/Runtime/SignalProcessing/Private/AllPassFractionalDelay.cpp
+++ b/Engine/Source/Runtime/SignalProcessing/Private/AllPassFractionalDelay.cpp
@@ -1,165 +1,166 @@
 // Copyright Epic Games, Inc. All Rights Reserved.
 
 #include "DSP/AllPassFractionalDelay.h"
-#include "DSP/Dsp.h"
 #include "DSP/BufferVectorOperations.h"
+#include "DSP/Dsp.h"
+#include "DSP/FloatArrayMath.h"
 
-using namespace Audio;
-
-FAllPassFractionalDelay::FAllPassFractionalDelay(int32 InMaxDelay, int32 InNumInternalBufferSamples)
-: MaxDelay(InMaxDelay)
-, NumInternalBufferSamples(InNumInternalBufferSamples)
-, Z1(0.0f)
+namespace Audio
 {
-	checkf(MaxDelay > 0, TEXT("Maximum delay must be greater than zero"));
-	checkf(InNumInternalBufferSamples > 0, TEXT("Internal buffer length must be greater than zero"));
-
-	// Clamp settings to avoid hard crashes during runtime.
-	if (MaxDelay < 1)
+	FAllPassFractionalDelay::FAllPassFractionalDelay(int32 InMaxDelay, int32 InNumInternalBufferSamples)
+	: MaxDelay(InMaxDelay)
+	, NumInternalBufferSamples(InNumInternalBufferSamples)
+	, Z1(0.0f)
 	{
-		MaxDelay = 1;
-	}
-	if (InNumInternalBufferSamples < 1)
-	{
-		InNumInternalBufferSamples = 32;
-	}
+		checkf(MaxDelay > 0, TEXT("Maximum delay must be greater than zero"));
+		checkf(InNumInternalBufferSamples > 0, TEXT("Internal buffer length must be greater than zero"));
+
+		// Clamp settings to avoid hard crashes during runtime.
+		if (MaxDelay < 1)
+		{
+			MaxDelay = 1;
+		}
+		if (InNumInternalBufferSamples < 1)
+		{
+			InNumInternalBufferSamples = 32;
+		}
 
-	// Allocate and prepare delay line for maximum delay.
-	DelayLine = MakeUnique<FAlignedBlockBuffer>((2 * MaxDelay) + NumInternalBufferSamples, MaxDelay + NumInternalBufferSamples);
-	DelayLine->AddZeros(MaxDelay);
+		// Allocate and prepare delay line for maximum delay.
+		DelayLine = MakeUnique<FAlignedBlockBuffer>((2 * MaxDelay) + NumInternalBufferSamples, MaxDelay + NumInternalBufferSamples);
+		DelayLine->AddZeros(MaxDelay);
 
-	Coefficients.Reset(NumInternalBufferSamples);
-	FractionalDelays.Reset(NumInternalBufferSamples);
-	IntegerDelays.Reset(NumInternalBufferSamples);
-	IntegerDelayOffsets.Reset(NumInternalBufferSamples);
+		Coefficients.Reset(NumInternalBufferSamples);
+		FractionalDelays.Reset(NumInternalBufferSamples);
+		IntegerDelays.Reset(NumInternalBufferSamples);
+		IntegerDelayOffsets.Reset(NumInternalBufferSamples);
 
-	Coefficients.AddUninitialized(NumInternalBufferSamples);
-	FractionalDelays.AddUninitialized(NumInternalBufferSamples);
-	IntegerDelays.AddUninitialized(NumInternalBufferSamples);
-	IntegerDelayOffsets.AddUninitialized(NumInternalBufferSamples);
+		Coefficients.AddUninitialized(NumInternalBufferSamples);
+		FractionalDelays.AddUninitialized(NumInternalBufferSamples);
+		IntegerDelays.AddUninitialized(NumInternalBufferSamples);
+		IntegerDelayOffsets.AddUninitialized(NumInternalBufferSamples);
 
-	// Integer delay offsets account for buffer position when doing block processing of data. 
-	for (int32 i = 0; i < InNumInternalBufferSamples; i++)
-	{
-		IntegerDelayOffsets[i] = i + MaxDelay;
+		// Integer delay offsets account for buffer position when doing block processing of data. 
+		for (int32 i = 0; i < InNumInternalBufferSamples; i++)
+		{
+			IntegerDelayOffsets[i] = i + MaxDelay;
+		}
 	}
-}
 
 
-// Destructor
-FAllPassFractionalDelay::~FAllPassFractionalDelay()
-{}
+	// Destructor
+	FAllPassFractionalDelay::~FAllPassFractionalDelay()
+	{}
 
-// Resets the delay line state, flushes buffer and resets read/write pointers.
-void FAllPassFractionalDelay::Reset() 
-{
-	DelayLine->ClearSamples();
-	DelayLine->AddZeros(MaxDelay);
-	Z1 = 0.0f;
-}
+	// Resets the delay line state, flushes buffer and resets read/write pointers.
+	void FAllPassFractionalDelay::Reset() 
+	{
+		DelayLine->ClearSamples();
+		DelayLine->AddZeros(MaxDelay);
+		Z1 = 0.0f;
+	}
 
 
-void FAllPassFractionalDelay::ProcessAudio(const FAlignedFloatBuffer& InSamples, const FAlignedFloatBuffer& InDelays, FAlignedFloatBuffer& OutSamples)
-{
-	const int32 InNum = InSamples.Num();
-	checkf(InNum == InDelays.Num(), TEXT("Input buffers must be equal length."));
+	void FAllPassFractionalDelay::ProcessAudio(const FAlignedFloatBuffer& InSamples, const FAlignedFloatBuffer& InDelays, FAlignedFloatBuffer& OutSamples)
+	{
+		const int32 InNum = InSamples.Num();
+		checkf(InNum == InDelays.Num(), TEXT("Input buffers must be equal length."));
 
-	// Prepare output buffer
-	OutSamples.Reset(InNum);
-	OutSamples.AddUninitialized(InNum);
+		// Prepare output buffer
+		OutSamples.Reset(InNum);
+		OutSamples.AddUninitialized(InNum);
 
-	float* OutSampleData = OutSamples.GetData();
-	const float* InSampleData = InSamples.GetData();
-	const float* InDelayData = InDelays.GetData();
+		float* OutSampleData = OutSamples.GetData();
+		const float* InSampleData = InSamples.GetData();
+		const float* InDelayData = InDelays.GetData();
 
-	// Process audio one block at a time.
-	int32 LeftOver = InNum;
-	int32 BufferPos = 0;
-	while (LeftOver > 0)
-	{
-		int32 NumToProcess = FMath::Min(LeftOver, NumInternalBufferSamples);
-		ProcessAudioBlock(&InSampleData[BufferPos], &InDelayData[BufferPos], NumToProcess, &OutSampleData[BufferPos]);
-		BufferPos += NumToProcess;
-		LeftOver -= NumToProcess;
+		// Process audio one block at a time.
+		int32 LeftOver = InNum;
+		int32 BufferPos = 0;
+		while (LeftOver > 0)
+		{
+			int32 NumToProcess = FMath::Min(LeftOver, NumInternalBufferSamples);
+			ProcessAudioBlock(&InSampleData[BufferPos], &InDelayData[BufferPos], NumToProcess, &OutSampleData[BufferPos]);
+			BufferPos += NumToProcess;
+			LeftOver -= NumToProcess;
+		}
 	}
-}
 
-void FAllPassFractionalDelay::ProcessAudioBlock(const float* InSamples, const float* InDelays, const int32 InNum, float* OutSamples)
-{
-	checkf(0 == (InNum % 4), TEXT("Array length must be multiple of 4"));
-
-	// All these assume that InNum <= NumInternalBufferSamples
-	const int32* IntegerDelayOffsetData = IntegerDelayOffsets.GetData();
-	int32* IntegerDelayData = IntegerDelays.GetData();
-	float* FractionalDelayData = FractionalDelays.GetData();
-	float* CoefficientsData = Coefficients.GetData();
-
-	FMemory::Memcpy(FractionalDelayData, InDelays, InNum * sizeof(float));
-	
-	// Determine integer delays and filter coefficients per a sample
-	const VectorRegister4Float VTwoAndAHalf = MakeVectorRegister(2.5f, 2.5f, 2.5f, 2.5f);
-	const VectorRegister4Float VMaxDelay = MakeVectorRegister((float)MaxDelay, (float)MaxDelay, (float)MaxDelay, (float)MaxDelay);
-	for (int32 i = 0; i < InNum; i += 4)
+	void FAllPassFractionalDelay::ProcessAudioBlock(const float* InSamples, const float* InDelays, const int32 InNum, float* OutSamples)
 	{
-		checkf(FractionalDelayData[i] <= MaxDelay, TEXT("Delay exceeds maximum"));
-		checkf(FractionalDelayData[i + 1] <= MaxDelay, TEXT("Delay exceeds maximum"));
-		checkf(FractionalDelayData[i + 2] <= MaxDelay, TEXT("Delay exceeds maximum"));
-		checkf(FractionalDelayData[i + 3] <= MaxDelay, TEXT("Delay exceeds maximum"));
-
-		// Ensure that delays are in [0.5, MaxDelay]
-		// Delay= Max(Delay, 0.5)
-		VectorRegister4Float VFractionalDelays = VectorLoadAligned(&FractionalDelayData[i]);
-		VFractionalDelays = VectorMax(VFractionalDelays, GlobalVectorConstants::FloatOneHalf);
-		// Delay = Min(Delay, MaxDelay)
-		VFractionalDelays = VectorMin(VFractionalDelays, VMaxDelay);
-
-		// To ensure fractional delays between -0.5 - 0.5, subtract limit then add back in.
-		// Delay = Delay + 0.5
-		VFractionalDelays = VectorAdd(VFractionalDelays, GlobalVectorConstants::FloatOneHalf);
-		// Delay.floor = Floor(Delay)
-		VectorRegister4Float VFloorDelays = VectorFloor(VFractionalDelays);
-		// Delay.frac = Delay - Delay.floor
-		VFractionalDelays = VectorSubtract(VFractionalDelays, VFloorDelays);
-		
-		// Reintroduing 0.5 previously removed.
-		// alpha = Delay.frac - 0.5
-		VectorRegister4Float VCoefficients = VectorSubtract(VFractionalDelays, GlobalVectorConstants::FloatOneHalf);
-		// denom = 2.5 - alpha
-		VectorRegister4Float VDenominator = VectorSubtract(VTwoAndAHalf, VFractionalDelays);
-		// coef = alpha / (2.5 - alpha)
-		VCoefficients = VectorDivide(VCoefficients, VDenominator);
-		VectorStoreAligned(VCoefficients, &Coefficients[i]);
-
-		// Delay.int = int(Delay.floor)
-		VectorRegister4Int VIntegerDelays = VectorFloatToInt(VFloorDelays);
-		VectorRegister4Int VIntegerDelayOffset = VectorIntLoadAligned(&IntegerDelayOffsetData[i]);
-		// Delay.int += BufferIdx
-		VectorIntStoreAligned(VectorIntSubtract(VIntegerDelayOffset, VIntegerDelays), &IntegerDelayData[i]);
-	}
+		checkf(0 == (InNum % 4), TEXT("Array length must be multiple of 4"));
 
-	// Update delay line.
-	DelayLine->AddSamples(InSamples, InNum);
-
-	// All pass filter delay line to get fractional delay
-	const float* DelayData = DelayLine->InspectSamples(InNum + MaxDelay);
-	int32 DelayPos;
-	if (InNum > 0)
-	{ 
-		DelayPos = IntegerDelayData[0];
-		OutSamples[0] = DelayData[DelayPos - 1] + CoefficientsData[0] * DelayData[DelayPos] - CoefficientsData[0] * Z1;
-		for (int32 i = 1; i < InNum; i++)
+		// All these assume that InNum <= NumInternalBufferSamples
+		const int32* IntegerDelayOffsetData = IntegerDelayOffsets.GetData();
+		int32* IntegerDelayData = IntegerDelays.GetData();
+		float* FractionalDelayData = FractionalDelays.GetData();
+		float* CoefficientsData = Coefficients.GetData();
+
+		FMemory::Memcpy(FractionalDelayData, InDelays, InNum * sizeof(float));
+		
+		// Determine integer delays and filter coefficients per a sample
+		const VectorRegister4Float VTwoAndAHalf = MakeVectorRegister(2.5f, 2.5f, 2.5f, 2.5f);
+		const VectorRegister4Float VMaxDelay = MakeVectorRegister((float)MaxDelay, (float)MaxDelay, (float)MaxDelay, (float)MaxDelay);
+		for (int32 i = 0; i < InNum; i += 4)
 		{
-			DelayPos = IntegerDelayData[i];
-			OutSamples[i] = DelayData[DelayPos - 1] + CoefficientsData[i] * DelayData[DelayPos] - CoefficientsData[i] * OutSamples[i - 1];
+			checkf(FractionalDelayData[i] <= MaxDelay, TEXT("Delay exceeds maximum"));
+			checkf(FractionalDelayData[i + 1] <= MaxDelay, TEXT("Delay exceeds maximum"));
+			checkf(FractionalDelayData[i + 2] <= MaxDelay, TEXT("Delay exceeds maximum"));
+			checkf(FractionalDelayData[i + 3] <= MaxDelay, TEXT("Delay exceeds maximum"));
+
+			// Ensure that delays are in [0.5, MaxDelay]
+			// Delay= Max(Delay, 0.5)
+			VectorRegister4Float VFractionalDelays = VectorLoad(&FractionalDelayData[i]);
+			VFractionalDelays = VectorMax(VFractionalDelays, GlobalVectorConstants::FloatOneHalf);
+			// Delay = Min(Delay, MaxDelay)
+			VFractionalDelays = VectorMin(VFractionalDelays, VMaxDelay);
+
+			// To ensure fractional delays between -0.5 - 0.5, subtract limit then add back in.
+			// Delay = Delay + 0.5
+			VFractionalDelays = VectorAdd(VFractionalDelays, GlobalVectorConstants::FloatOneHalf);
+			// Delay.floor = Floor(Delay)
+			VectorRegister4Float VFloorDelays = VectorFloor(VFractionalDelays);
+			// Delay.frac = Delay - Delay.floor
+			VFractionalDelays = VectorSubtract(VFractionalDelays, VFloorDelays);
+			
+			// Reintroduing 0.5 previously removed.
+			// alpha = Delay.frac - 0.5
+			VectorRegister4Float VCoefficients = VectorSubtract(VFractionalDelays, GlobalVectorConstants::FloatOneHalf);
+			// denom = 2.5 - alpha
+			VectorRegister4Float VDenominator = VectorSubtract(VTwoAndAHalf, VFractionalDelays);
+			// coef = alpha / (2.5 - alpha)
+			VCoefficients = VectorDivide(VCoefficients, VDenominator);
+			VectorStore(VCoefficients, &Coefficients[i]);
+
+			// Delay.int = int(Delay.floor)
+			VectorRegister4Int VIntegerDelays = VectorFloatToInt(VFloorDelays);
+			VectorRegister4Int VIntegerDelayOffset = VectorIntLoadAligned(&IntegerDelayOffsetData[i]);
+			// Delay.int += BufferIdx
+			VectorIntStoreAligned(VectorIntSubtract(VIntegerDelayOffset, VIntegerDelays), &IntegerDelayData[i]);
 		}
 
-		TArrayView<float> OutSamplesView(OutSamples, InNum);
-		ArrayUnderflowClamp(OutSamplesView);
-
-		Z1 = OutSamples[InNum - 1];
+		// Update delay line.
+		DelayLine->AddSamples(InSamples, InNum);
+
+		// All pass filter delay line to get fractional delay
+		const float* DelayData = DelayLine->InspectSamples(InNum + MaxDelay);
+		int32 DelayPos;
+		if (InNum > 0)
+		{ 
+			DelayPos = IntegerDelayData[0];
+			OutSamples[0] = DelayData[DelayPos - 1] + CoefficientsData[0] * DelayData[DelayPos] - CoefficientsData[0] * Z1;
+			for (int32 i = 1; i < InNum; i++)
+			{
+				DelayPos = IntegerDelayData[i];
+				OutSamples[i] = DelayData[DelayPos - 1] + CoefficientsData[i] * DelayData[DelayPos] - CoefficientsData[i] * OutSamples[i - 1];
+			}
+
+			TArrayView<float> OutSamplesView(OutSamples, InNum);
+			ArrayUnderflowClamp(OutSamplesView);
+
+			Z1 = OutSamples[InNum - 1];
+		}
+		
+		// Remove unneeded delay line.
+		DelayLine->RemoveSamples(InNum);
 	}
-	
-	// Remove unneeded delay line.
-	DelayLine->RemoveSamples(InNum);
 }
-
diff --git a/Engine/Source/Runtime/SignalProcessing/Private/AudioBufferDistanceAttenuation.cpp b/Engine/Source/Runtime/SignalProcessing/Private/AudioBufferDistanceAttenuation.cpp
--- a/Engine/Source/Runtime/SignalProcessing/Private/AudioBufferDistanceAttenuation.cpp
+++ b/Engine/Source/Runtime/SignalProcessing/Private/AudioBufferDistanceAttenuation.cpp
@@ -2,7 +2,7 @@
 
 #include "DSP/AudioBufferDistanceAttenuation.h"
 #include "DSP/Dsp.h"
-#include "DSP/BufferVectorOperations.h"
+#include "DSP/FloatArrayMath.h"
 
 namespace Audio
 {
diff --git a/Engine/Source/Runtime/SignalProcessing/Private/AudioDebuggingUtilities.cpp b/Engine/Source/Runtime/SignalProcessing/Private/AudioDebuggingUtilities.cpp
new file mode 100644
--- /dev/null
+++ b/Engine/Source/Runtime/SignalProcessing/Private/AudioDebuggingUtilities.cpp
@@ -0,0 +1,33 @@
+// Copyright Epic Games, Inc. All Rights Reserved.
+
+#include "DSP/AudioDebuggingUtilities.h"
+
+#include "DSP/Dsp.h"
+#include "DSP/FloatArrayMath.h"
+#include "HAL/Platform.h"
+
+void BreakWhenAudible(float* InBuffer, int32 NumSamples)
+{
+	static const float AudibilityThreshold = Audio::ConvertToLinear(-40.0f);
+
+	TArrayView<const float> InBufferView(InBuffer, NumSamples);
+	float BufferAmplitude = Audio::ArrayGetAverageAbsValue(InBufferView);
+
+	if (BufferAmplitude > AudibilityThreshold)
+	{
+		PLATFORM_BREAK();
+	}
+}
+
+void BreakWhenTooLoud(float* InBuffer, int32 NumSamples)
+{
+	static const float PainThreshold = Audio::ConvertToLinear(3.0f);
+
+	TArrayView<const float> InBufferView(InBuffer, NumSamples);
+	float BufferAmplitude = Audio::ArrayGetAverageAbsValue(InBufferView);
+
+	if (BufferAmplitude > PainThreshold)
+	{
+		PLATFORM_BREAK();
+	}
+}
diff --git a/Engine/Source/Runtime/SignalProcessing/Private/AudioFFT.cpp b/Engine/Source/Runtime/SignalProcessing/Private/AudioFFT.cpp
--- a/Engine/Source/Runtime/SignalProcessing/Private/AudioFFT.cpp
+++ b/Engine/Source/Runtime/SignalProcessing/Private/AudioFFT.cpp
@@ -1,8 +1,10 @@
 // Copyright Epic Games, Inc. All Rights Reserved.
 
 #include "DSP/AudioFFT.h"
-#include "HAL/IConsoleManager.h"
+
 #include "DSP/BufferVectorOperations.h"
+#include "DSP/FloatArrayMath.h"
+#include "HAL/IConsoleManager.h"
 
 #define IFFT_PRESERVE_COMPLEX_COMPONENT 0
 
@@ -104,6 +106,71 @@ namespace Audio
 		}
 	}
 
+	FWindow::FWindow(EWindowType InType, int32 InNumFrames, int32 InNumChannels, bool bIsPeriodic)
+		: WindowType(InType)
+		, NumSamples(InNumFrames * InNumChannels)
+	{
+		checkf(NumSamples % 4 == 0, TEXT("For performance reasons, this window's length should be a multiple of 4."));
+		Generate(InNumFrames, InNumChannels, bIsPeriodic);
+	}
+
+	// Apply this window to InBuffer, which is expected to be an interleaved buffer with the same amount of frames
+	// and channels this window was constructed with.
+	void FWindow::ApplyToBuffer(float* InBuffer)
+	{
+		if (WindowType == EWindowType::None)
+		{
+			return;
+		}
+
+		TArrayView<const float> WindowBufferView(WindowBuffer.GetData(), NumSamples);
+		TArrayView<float> InBufferView(InBuffer, NumSamples);
+		ArrayMultiplyInPlace(WindowBufferView, InBufferView);
+	}
+
+	EWindowType FWindow::GetWindowType() const
+	{
+		return WindowType;
+	}
+
+
+	// Generate the window. Called on constructor.
+	void FWindow::Generate(int32 NumFrames, int32 NumChannels, bool bIsPeriodic)
+	{
+		if (WindowType == EWindowType::None)
+		{
+			return;
+		}
+
+		WindowBuffer.Reset();
+		WindowBuffer.AddZeroed(NumSamples);
+
+		switch (WindowType)
+		{
+			case EWindowType::Hann:
+			{
+				GenerateHannWindow(WindowBuffer.GetData(), NumFrames, NumChannels, bIsPeriodic);
+				break;
+			}
+			case EWindowType::Hamming:
+			{
+				GenerateHammingWindow(WindowBuffer.GetData(), NumFrames, NumChannels, bIsPeriodic);
+				break;
+			}
+			case EWindowType::Blackman:
+			{
+				GenerateBlackmanWindow(WindowBuffer.GetData(), NumFrames, NumChannels, bIsPeriodic);
+				break;
+			}
+			default:
+			{
+				checkf(false, TEXT("Unknown window type!"));
+				break;
+			}
+		}
+	}
+
+
 	namespace FFTIntrinsics
 	{
 		 uint32 NextPowerOf2(uint32 Input)
diff --git a/Engine/Source/Runtime/SignalProcessing/Private/BufferOnePoleLPF.cpp b/Engine/Source/Runtime/SignalProcessing/Private/BufferOnePoleLPF.cpp
--- a/Engine/Source/Runtime/SignalProcessing/Private/BufferOnePoleLPF.cpp
+++ b/Engine/Source/Runtime/SignalProcessing/Private/BufferOnePoleLPF.cpp
@@ -1,71 +1,70 @@
 // Copyright Epic Games, Inc. All Rights Reserved.
 
 #include "DSP/BufferOnePoleLPF.h"
-#include "DSP/BufferVectorOperations.h"
+#include "DSP/FloatArrayMath.h"
 
-using namespace Audio;
-
-// Constructor 
-FBufferOnePoleLPF::FBufferOnePoleLPF(float InG)
-	: CutoffFrequency(0.0f)
-	, B1(0.0f)
-	, A0(1.0f)
-	, Z1(0.0f)
-{
-	SetG(InG);
-}
-
-// Set the LPF gain coefficient
-void FBufferOnePoleLPF::SetG(float InG)
-{ 
-	B1 = InG;
-	A0 = 1.0f - B1;
-}
-
-// Resets the sample delay to 0
-void FBufferOnePoleLPF::Reset()
+namespace Audio
 {
-	SetG(0.0f);
-	FlushAudio();
-}
+	// Constructor 
+	FBufferOnePoleLPF::FBufferOnePoleLPF(float InG)
+		: CutoffFrequency(0.0f)
+		, B1(0.0f)
+		, A0(1.0f)
+		, Z1(0.0f)
+	{
+		SetG(InG);
+	}
 
-void FBufferOnePoleLPF::FlushAudio()
-{
-	Z1 = 0.0f;
-}
+	// Set the LPF gain coefficient
+	void FBufferOnePoleLPF::SetG(float InG)
+	{ 
+		B1 = InG;
+		A0 = 1.0f - B1;
+	}
 
-/** Sets the filter frequency using normalized frequency (between 0.0 and 1.0f or 0.0 hz and Nyquist Frequency in Hz) */
-void FBufferOnePoleLPF::SetFrequency(const float InFrequency)
-{
-	if (!FMath::IsNearlyEqual(InFrequency, CutoffFrequency))
+	// Resets the sample delay to 0
+	void FBufferOnePoleLPF::Reset()
 	{
-		CutoffFrequency = InFrequency;
-		SetG(FMath::Exp(-PI * CutoffFrequency));
+		SetG(0.0f);
+		FlushAudio();
 	}
-}
 
-void FBufferOnePoleLPF::ProcessAudio(const Audio::FAlignedFloatBuffer& InSamples, Audio::FAlignedFloatBuffer& OutSamples)
-{
-	int32 Index;
-	const int32 InNum = InSamples.Num();
-	const float* InSampleData = InSamples.GetData();
-	
-	// Prepare output samples
-	OutSamples.Reset(InNum);
-	OutSamples.AddUninitialized(InNum);
-	float* OutSampleData = OutSamples.GetData();
+	void FBufferOnePoleLPF::FlushAudio()
+	{
+		Z1 = 0.0f;
+	}
 
-	// Filter audio
-	int32 DelayIndex;
-	OutSampleData[0] = InSampleData[0] * A0 + B1 * Z1;
-	for(Index = 1, DelayIndex = 0; Index < InNum; Index++, DelayIndex++)
+	/** Sets the filter frequency using normalized frequency (between 0.0 and 1.0f or 0.0 hz and Nyquist Frequency in Hz) */
+	void FBufferOnePoleLPF::SetFrequency(const float InFrequency)
 	{
-		//OutSampleData[Index] = UnderflowClamp(InSampleData[Index] * A0 + B1 * OutSampleData[DelayIndex]);
-		OutSampleData[Index] = InSampleData[Index] * A0 + B1 * OutSampleData[DelayIndex];
+		if (!FMath::IsNearlyEqual(InFrequency, CutoffFrequency))
+		{
+			CutoffFrequency = InFrequency;
+			SetG(FMath::Exp(-PI * CutoffFrequency));
+		}
 	}
-	ArrayUnderflowClamp(OutSamples);
-	// Store delay value
-	Z1 = OutSampleData[InNum - 1];
-}
 
+	void FBufferOnePoleLPF::ProcessAudio(const Audio::FAlignedFloatBuffer& InSamples, Audio::FAlignedFloatBuffer& OutSamples)
+	{
+		int32 Index;
+		const int32 InNum = InSamples.Num();
+		const float* InSampleData = InSamples.GetData();
+		
+		// Prepare output samples
+		OutSamples.Reset(InNum);
+		OutSamples.AddUninitialized(InNum);
+		float* OutSampleData = OutSamples.GetData();
 
+		// Filter audio
+		int32 DelayIndex;
+		OutSampleData[0] = InSampleData[0] * A0 + B1 * Z1;
+		for(Index = 1, DelayIndex = 0; Index < InNum; Index++, DelayIndex++)
+		{
+			//OutSampleData[Index] = UnderflowClamp(InSampleData[Index] * A0 + B1 * OutSampleData[DelayIndex]);
+			OutSampleData[Index] = InSampleData[Index] * A0 + B1 * OutSampleData[DelayIndex];
+		}
+		ArrayUnderflowClamp(OutSamples);
+		// Store delay value
+		Z1 = OutSampleData[InNum - 1];
+	}
+}
diff --git a/Engine/Source/Runtime/SignalProcessing/Private/BufferVectorOperations.cpp b/Engine/Source/Runtime/SignalProcessing/Private/BufferVectorOperations.cpp
--- a/Engine/Source/Runtime/SignalProcessing/Private/BufferVectorOperations.cpp
+++ b/Engine/Source/Runtime/SignalProcessing/Private/BufferVectorOperations.cpp
@@ -75,9 +75,9 @@ namespace Audio
 		
 		for (int32 i = 0; i < NumSamples; i += AUDIO_NUM_FLOATS_PER_VECTOR_REGISTER)
 		{
-			VectorRegister4Float Result = VectorLoadAligned(&StereoBuffer[i]);
+			VectorRegister4Float Result = VectorLoad(&StereoBuffer[i]);
 			Result = VectorMultiply(Result, GainVector);
-			VectorStoreAligned(Result, &StereoBuffer[i]);
+			VectorStore(Result, &StereoBuffer[i]);
 		}
 	}
 
@@ -96,9 +96,9 @@ namespace Audio
 
 		for (int32 i = 0; i < NumSamples; i += AUDIO_NUM_FLOATS_PER_VECTOR_REGISTER)
 		{
-			VectorRegister4Float Result = VectorLoadAligned(&StereoBuffer[i]);
+			VectorRegister4Float Result = VectorLoad(&StereoBuffer[i]);
 			Result = VectorMultiply(Result, GainVector);
-			VectorStoreAligned(Result, &StereoBuffer[i]);
+			VectorStore(Result, &StereoBuffer[i]);
 
 			GainVector = VectorAdd(GainVector, GainDeltasVector);
 		}
@@ -129,7 +129,7 @@ namespace Audio
 		{
 			VectorRegister4Float Result = VectorSet(MonoBuffer[i], MonoBuffer[i], MonoBuffer[i + 1], MonoBuffer[i + 1]);
 			Result = VectorMultiply(Result, GainVector);
-			VectorStoreAligned(Result, &DestinationBuffer[i*2]);
+			VectorStore(Result, &DestinationBuffer[i*2]);
 		}
 	}
 
@@ -167,7 +167,7 @@ namespace Audio
 		{
 			VectorRegister4Float Result = VectorSet(MonoBuffer[i], MonoBuffer[i], MonoBuffer[i + 1], MonoBuffer[i + 1]);
 			Result = VectorMultiply(Result, GainVector);
-			VectorStoreAligned(Result, &DestinationBuffer[i*2]);
+			VectorStore(Result, &DestinationBuffer[i*2]);
 
 			GainVector = VectorAdd(GainVector, GainDeltasVector);
 		}
@@ -193,18 +193,16 @@ namespace Audio
 	{
 		checkf(InNumFrames >= AUDIO_NUM_FLOATS_PER_VECTOR_REGISTER, TEXT("Buffer must have at least %i elements."), AUDIO_NUM_FLOATS_PER_VECTOR_REGISTER);
 		checkf(0 == (InNumFrames % AUDIO_NUM_FLOATS_PER_VECTOR_REGISTER), TEXT("Buffer length be a multiple of %i."), AUDIO_NUM_FLOATS_PER_VECTOR_REGISTER);
-		checkf(IsAligned<const float*>(MonoBuffer, AUDIO_SIMD_BYTE_ALIGNMENT), TEXT("Memory must be aligned to use vector operations."));
-		checkf(IsAligned<float*>(DestinationBuffer, AUDIO_SIMD_BYTE_ALIGNMENT), TEXT("Memory must be aligned to use vector operations."));
 
 		int32 OutPos = 0;
 		for (int32 i = 0; i < InNumFrames; i += AUDIO_NUM_FLOATS_PER_VECTOR_REGISTER)
 		{
-			VectorRegister4Float Input = VectorLoadAligned(&MonoBuffer[i]);
+			VectorRegister4Float Input = VectorLoad(&MonoBuffer[i]);
 			VectorRegister4Float Output = VectorSwizzle(Input, 0, 0, 1, 1);
-			VectorStoreAligned(Output, &DestinationBuffer[OutPos]);
+			VectorStore(Output, &DestinationBuffer[OutPos]);
 			OutPos += AUDIO_NUM_FLOATS_PER_VECTOR_REGISTER;
 			Output = VectorSwizzle(Input, 2, 2, 3, 3);
-			VectorStoreAligned(Output, &DestinationBuffer[OutPos]);
+			VectorStore(Output, &DestinationBuffer[OutPos]);
 			OutPos += AUDIO_NUM_FLOATS_PER_VECTOR_REGISTER;
 		}
 	}
@@ -245,7 +243,7 @@ namespace Audio
 			VectorRegister4Float Result = VectorMultiply(Input1, GainVector1);
 			Result = VectorMultiplyAdd(Input2, GainVector2, Result);
 
-			VectorStoreAligned(Result, &DestinationBuffer[i * 2]);
+			VectorStore(Result, &DestinationBuffer[i * 2]);
 		}
 	}
 
@@ -299,7 +297,7 @@ namespace Audio
 			VectorRegister4Float Result = VectorMultiply(Input1, GainVector1);
 			Result = VectorMultiplyAdd(Input2, GainVector2, Result);
 
-			VectorStoreAligned(Result, &DestinationBuffer[i * 2]);
+			VectorStore(Result, &DestinationBuffer[i * 2]);
 
 			GainVector1 = VectorAdd(GainVector1, GainDeltasVector1);
 			GainVector2 = VectorAdd(GainVector2, GainDeltasVector2);
@@ -313,13 +311,13 @@ namespace Audio
 
 	void Apply4ChannelGain(float* RESTRICT InterleavedBuffer, int32 NumSamples, const float* RESTRICT Gains)
 	{
-		const VectorRegister4Float GainVector = VectorLoadAligned(Gains);
+		const VectorRegister4Float GainVector = VectorLoad(Gains);
 
 		for (int32 i = 0; i < NumSamples; i += AUDIO_NUM_FLOATS_PER_VECTOR_REGISTER)
 		{
-			VectorRegister4Float Result = VectorLoadAligned(&InterleavedBuffer[i]);
+			VectorRegister4Float Result = VectorLoad(&InterleavedBuffer[i]);
 			Result = VectorMultiply(Result, GainVector);
-			VectorStoreAligned(Result, &InterleavedBuffer[i]);
+			VectorStore(Result, &InterleavedBuffer[i]);
 		}
 	}
 
@@ -331,16 +329,16 @@ namespace Audio
 	void Apply4ChannelGain(float* RESTRICT InterleavedBuffer, int32 NumSamples, const float* RESTRICT StartGains, const float* RESTRICT EndGains)
 	{
 		// Initialize GainVector at StartGains and compute GainDeltasVector:
-		VectorRegister4Float GainVector = VectorLoadAligned(StartGains);
-		const VectorRegister4Float DestinationVector = VectorLoadAligned(EndGains);
+		VectorRegister4Float GainVector = VectorLoad(StartGains);
+		const VectorRegister4Float DestinationVector = VectorLoad(EndGains);
 		const VectorRegister4Float NumFramesVector = VectorSetFloat1(NumSamples / (float)AUDIO_NUM_FLOATS_PER_VECTOR_REGISTER);
 		const VectorRegister4Float GainDeltasVector = VectorDivide(VectorSubtract(DestinationVector, GainVector), NumFramesVector);
 
 		for (int32 i = 0; i < NumSamples; i += AUDIO_NUM_FLOATS_PER_VECTOR_REGISTER)
 		{
-			VectorRegister4Float Result = VectorLoadAligned(&InterleavedBuffer[i]);
+			VectorRegister4Float Result = VectorLoad(&InterleavedBuffer[i]);
 			Result = VectorMultiply(Result, GainVector);
-			VectorStoreAligned(Result, &InterleavedBuffer[i]);
+			VectorStore(Result, &InterleavedBuffer[i]);
 
 			GainVector = VectorAdd(GainVector, GainDeltasVector);
 		}
@@ -367,13 +365,13 @@ namespace Audio
 
 	void MixMonoTo4ChannelsFast(const float* RESTRICT MonoBuffer, float* RESTRICT DestinationBuffer, int32 NumFrames, const float* RESTRICT Gains)
 	{
-		const VectorRegister4Float GainVector = VectorLoadAligned(Gains);
+		const VectorRegister4Float GainVector = VectorLoad(Gains);
 
 		for (int32 i = 0; i < NumFrames; i++)
 		{
 			VectorRegister4Float Result = VectorLoadFloat1(&MonoBuffer[i]);
 			Result = VectorMultiply(Result, GainVector);
-			VectorStoreAligned(Result, &DestinationBuffer[i*4]);
+			VectorStore(Result, &DestinationBuffer[i*4]);
 		}
 	}
 
@@ -398,8 +396,8 @@ namespace Audio
 	*/
 	void MixMonoTo4ChannelsFast(const float* RESTRICT MonoBuffer, float* RESTRICT DestinationBuffer, int32 NumFrames, const float* RESTRICT StartGains, const float* RESTRICT EndGains)
 	{
-		VectorRegister4Float GainVector = VectorLoadAligned(StartGains);
-		const VectorRegister4Float DestinationVector = VectorLoadAligned(EndGains);
+		VectorRegister4Float GainVector = VectorLoad(StartGains);
+		const VectorRegister4Float DestinationVector = VectorLoad(EndGains);
 		const VectorRegister4Float NumFramesVector = VectorSetFloat1((float) NumFrames);
 		const VectorRegister4Float GainDeltasVector = VectorDivide(VectorSubtract(DestinationVector, GainVector), NumFramesVector);
 
@@ -407,7 +405,7 @@ namespace Audio
 		{
 			VectorRegister4Float Result = VectorLoadFloat1(&MonoBuffer[i]);
 			Result = VectorMultiply(Result, GainVector);
-			VectorStoreAligned(Result, &DestinationBuffer[i * 4]);
+			VectorStore(Result, &DestinationBuffer[i * 4]);
 
 			GainVector = VectorAdd(GainVector, GainDeltasVector);
 		}
@@ -438,8 +436,8 @@ namespace Audio
 
 	void Mix2ChannelsTo4ChannelsFast(const float* RESTRICT SourceBuffer, float* RESTRICT DestinationBuffer, int32 NumFrames, const float* RESTRICT Gains)
 	{
-		const VectorRegister4Float GainVector1 = VectorLoadAligned(Gains);
-		const VectorRegister4Float GainVector2 = VectorLoadAligned(Gains + 4);
+		const VectorRegister4Float GainVector1 = VectorLoad(Gains);
+		const VectorRegister4Float GainVector2 = VectorLoad(Gains + 4);
 
 		for (int32 i = 0; i < NumFrames; i++)
 		{
@@ -448,7 +446,7 @@ namespace Audio
 
 			VectorRegister4Float Result = VectorMultiply(Input1, GainVector1);
 			Result = VectorMultiplyAdd(Input2, GainVector2, Result);
-			VectorStoreAligned(Result, &DestinationBuffer[i * 4]);
+			VectorStore(Result, &DestinationBuffer[i * 4]);
 		}
 	}
 
@@ -478,12 +476,12 @@ namespace Audio
 	{
 		const VectorRegister4Float NumFramesVector = VectorSetFloat1((float) NumFrames);
 
-		VectorRegister4Float GainVector1 = VectorLoadAligned(StartGains);
-		const VectorRegister4Float DestinationVector1 = VectorLoadAligned(EndGains);
+		VectorRegister4Float GainVector1 = VectorLoad(StartGains);
+		const VectorRegister4Float DestinationVector1 = VectorLoad(EndGains);
 		const VectorRegister4Float GainDeltasVector1 = VectorDivide(VectorSubtract(DestinationVector1, GainVector1), NumFramesVector);
 
-		VectorRegister4Float GainVector2 = VectorLoadAligned(StartGains + 4);
-		const VectorRegister4Float DestinationVector2 = VectorLoadAligned(EndGains + 4);
+		VectorRegister4Float GainVector2 = VectorLoad(StartGains + 4);
+		const VectorRegister4Float DestinationVector2 = VectorLoad(EndGains + 4);
 		const VectorRegister4Float GainDeltasVector2 = VectorDivide(VectorSubtract(DestinationVector2, GainVector2), NumFramesVector);
 
 		for (int32 i = 0; i < NumFrames; i++)
@@ -493,7 +491,7 @@ namespace Audio
 
 			VectorRegister4Float Result = VectorMultiply(Input1, GainVector1);
 			Result = VectorMultiplyAdd(Input2, GainVector2, Result);
-			VectorStoreAligned(Result, &DestinationBuffer[i * 4]);
+			VectorStore(Result, &DestinationBuffer[i * 4]);
 
 			GainVector1 = VectorAdd(GainVector1, GainDeltasVector1);
 			GainVector2 = VectorAdd(GainVector2, GainDeltasVector2);
@@ -507,23 +505,23 @@ namespace Audio
 
 	void Apply6ChannelGain(float* RESTRICT InterleavedBuffer, int32 NumSamples, const float* RESTRICT Gains)
 	{
-		const VectorRegister4Float GainVector1 = VectorLoadAligned(Gains);
+		const VectorRegister4Float GainVector1 = VectorLoad(Gains);
 		const VectorRegister4Float GainVector2 = VectorSet(Gains[4], Gains[5], Gains[0], Gains[1]);
 		const VectorRegister4Float GainVector3 = VectorLoad(&Gains[2]);
 
 		for (int32 i = 0; i < NumSamples; i += 12)
 		{
-			VectorRegister4Float Result = VectorLoadAligned(&InterleavedBuffer[i]);
+			VectorRegister4Float Result = VectorLoad(&InterleavedBuffer[i]);
 			Result = VectorMultiply(Result, GainVector1);
-			VectorStoreAligned(Result, &InterleavedBuffer[i]);
+			VectorStore(Result, &InterleavedBuffer[i]);
 
-			Result = VectorLoadAligned(&InterleavedBuffer[i + 4]);
+			Result = VectorLoad(&InterleavedBuffer[i + 4]);
 			Result = VectorMultiply(Result, GainVector2);
-			VectorStoreAligned(Result, &InterleavedBuffer[i + 4]);
+			VectorStore(Result, &InterleavedBuffer[i + 4]);
 
-			Result = VectorLoadAligned(&InterleavedBuffer[i + 8]);
+			Result = VectorLoad(&InterleavedBuffer[i + 8]);
 			Result = VectorMultiply(Result, GainVector3);
-			VectorStoreAligned(Result, &InterleavedBuffer[i + 8]);
+			VectorStore(Result, &InterleavedBuffer[i + 8]);
 		}
 	}
 
@@ -536,8 +534,8 @@ namespace Audio
 	{
 		const VectorRegister4Float NumFramesVector = VectorSetFloat1(NumSamples / 12.0f);
 
-		VectorRegister4Float GainVector1 = VectorLoadAligned(StartGains);
-		const VectorRegister4Float DestinationVector1 = VectorLoadAligned(EndGains);
+		VectorRegister4Float GainVector1 = VectorLoad(StartGains);
+		const VectorRegister4Float DestinationVector1 = VectorLoad(EndGains);
 		const VectorRegister4Float GainDeltasVector1 = VectorDivide(VectorSubtract(DestinationVector1, GainVector1), NumFramesVector);
 
 		VectorRegister4Float GainVector2 = VectorSet(StartGains[4], StartGains[5], StartGains[0], StartGains[1]);
@@ -550,21 +548,21 @@ namespace Audio
 
 		for (int32 i = 0; i < NumSamples; i += 12)
 		{
-			VectorRegister4Float Result = VectorLoadAligned(&InterleavedBuffer[i]);
+			VectorRegister4Float Result = VectorLoad(&InterleavedBuffer[i]);
 			Result = VectorMultiply(Result, GainVector1);
-			VectorStoreAligned(Result, &InterleavedBuffer[i]);
+			VectorStore(Result, &InterleavedBuffer[i]);
 
 			GainVector1 = VectorAdd(GainVector1, GainDeltasVector1);
 
-			Result = VectorLoadAligned(&InterleavedBuffer[i + 4]);
+			Result = VectorLoad(&InterleavedBuffer[i + 4]);
 			Result = VectorMultiply(Result, GainVector2);
-			VectorStoreAligned(Result, &InterleavedBuffer[i + 4]);
+			VectorStore(Result, &InterleavedBuffer[i + 4]);
 
 			GainVector2 = VectorAdd(GainVector2, GainDeltasVector2);
 
-			Result = VectorLoadAligned(&InterleavedBuffer[i + 8]);
+			Result = VectorLoad(&InterleavedBuffer[i + 8]);
 			Result = VectorMultiply(Result, GainVector3);
-			VectorStoreAligned(Result, &InterleavedBuffer[i + 8]);
+			VectorStore(Result, &InterleavedBuffer[i + 8]);
 
 			GainVector3 = VectorAdd(GainVector3, GainDeltasVector3);
 		}
@@ -592,7 +590,7 @@ namespace Audio
 
 	void MixMonoTo6ChannelsFast(const float* RESTRICT MonoBuffer, float* RESTRICT DestinationBuffer, int32 NumFrames, const float* RESTRICT Gains)
 	{
-		const VectorRegister4Float GainVector1 = VectorLoadAligned(Gains);
+		const VectorRegister4Float GainVector1 = VectorLoad(Gains);
 		const VectorRegister4Float GainVector2 = VectorSet(Gains[4], Gains[5], Gains[0], Gains[1]);
 		const VectorRegister4Float GainVector3 = VectorLoad(&Gains[2]);
 
@@ -603,13 +601,13 @@ namespace Audio
 			const VectorRegister4Float Input3 = VectorLoadFloat1(&MonoBuffer[i + 1]);
 
 			VectorRegister4Float Result = VectorMultiply(Input1, GainVector1);
-			VectorStoreAligned(Result, &DestinationBuffer[i * 6]);
+			VectorStore(Result, &DestinationBuffer[i * 6]);
 
 			Result = VectorMultiply(Input2, GainVector2);
-			VectorStoreAligned(Result, &DestinationBuffer[i * 6 + 4]);
+			VectorStore(Result, &DestinationBuffer[i * 6 + 4]);
 
 			Result = VectorMultiply(Input3, GainVector3);
-			VectorStoreAligned(Result, &DestinationBuffer[i * 6 + 8]);
+			VectorStore(Result, &DestinationBuffer[i * 6 + 8]);
 		}
 	}
 
@@ -637,8 +635,8 @@ namespace Audio
 	{
 		const VectorRegister4Float NumFramesVector = VectorSetFloat1(NumFrames / 2.0f);
 
-		VectorRegister4Float GainVector1 = VectorLoadAligned(StartGains);
-		const VectorRegister4Float DestinationVector1 = VectorLoadAligned(EndGains);
+		VectorRegister4Float GainVector1 = VectorLoad(StartGains);
+		const VectorRegister4Float DestinationVector1 = VectorLoad(EndGains);
 		const VectorRegister4Float GainDeltasVector1 = VectorDivide(VectorSubtract(DestinationVector1, GainVector1), NumFramesVector);
 
 		VectorRegister4Float GainVector2 = VectorSet(StartGains[4], StartGains[5], StartGains[0], StartGains[1]);
@@ -656,17 +654,17 @@ namespace Audio
 			const VectorRegister4Float Input3 = VectorLoadFloat1(&MonoBuffer[i + 1]);
 
 			VectorRegister4Float Result = VectorMultiply(Input1, GainVector1);
-			VectorStoreAligned(Result, &DestinationBuffer[i * 6]);
+			VectorStore(Result, &DestinationBuffer[i * 6]);
 
 			GainVector1 = VectorAdd(GainVector1, GainDeltasVector1);
 
 			Result = VectorMultiply(Input2, GainVector2);
-			VectorStoreAligned(Result, &DestinationBuffer[i * 6 + 4]);
+			VectorStore(Result, &DestinationBuffer[i * 6 + 4]);
 
 			GainVector2 = VectorAdd(GainVector2, GainDeltasVector2);
 
 			Result = VectorMultiply(Input3, GainVector3);
-			VectorStoreAligned(Result, &DestinationBuffer[i * 6 + 8]);
+			VectorStore(Result, &DestinationBuffer[i * 6 + 8]);
 
 			GainVector3 = VectorAdd(GainVector3, GainDeltasVector3);
 		}
@@ -698,13 +696,13 @@ namespace Audio
 
 	void Mix2ChannelsTo6ChannelsFast(const float* RESTRICT SourceBuffer, float* RESTRICT DestinationBuffer, int32 NumFrames, const float* RESTRICT Gains)
 	{
-		const VectorRegister4Float GainVector11 = VectorLoadAligned(Gains);
+		const VectorRegister4Float GainVector11 = VectorLoad(Gains);
 		const VectorRegister4Float GainVector21 = VectorSet(Gains[4], Gains[5], Gains[0], Gains[1]);
 		const VectorRegister4Float GainVector31 = VectorLoad(&Gains[2]);
 
 		const VectorRegister4Float GainVector12 = VectorLoad(Gains + 6);
 		const VectorRegister4Float GainVector22 = VectorSet(Gains[10], Gains[11], Gains[6], Gains[7]);
-		const VectorRegister4Float GainVector32 = VectorLoadAligned(&Gains[8]);
+		const VectorRegister4Float GainVector32 = VectorLoad(&Gains[8]);
 
 		for (int32 FrameIndex = 0; FrameIndex < NumFrames; FrameIndex += 2)
 		{
@@ -721,15 +719,15 @@ namespace Audio
 
 			VectorRegister4Float Result = VectorMultiply(Input11, GainVector11);
 			Result = VectorMultiplyAdd(Input12, GainVector12, Result);
-			VectorStoreAligned(Result, &DestinationBuffer[OutputIndex]);
+			VectorStore(Result, &DestinationBuffer[OutputIndex]);
 
 			Result = VectorMultiply(Input21, GainVector21);
 			Result = VectorMultiplyAdd(Input22, GainVector22, Result);
-			VectorStoreAligned(Result, &DestinationBuffer[OutputIndex + 4]);
+			VectorStore(Result, &DestinationBuffer[OutputIndex + 4]);
 
 			Result = VectorMultiply(Input31, GainVector31);
 			Result = VectorMultiplyAdd(Input32, GainVector32, Result);
-			VectorStoreAligned(Result, &DestinationBuffer[OutputIndex + 8]);
+			VectorStore(Result, &DestinationBuffer[OutputIndex + 8]);
 		}
 	}
 
@@ -761,8 +759,8 @@ namespace Audio
 	{
 		const VectorRegister4Float NumFramesVector = VectorSetFloat1(NumFrames / 2.0f);
 
-		VectorRegister4Float GainVector11 = VectorLoadAligned(StartGains);
-		const VectorRegister4Float DestinationVector11 = VectorLoadAligned(EndGains);
+		VectorRegister4Float GainVector11 = VectorLoad(StartGains);
+		const VectorRegister4Float DestinationVector11 = VectorLoad(EndGains);
 		const VectorRegister4Float GainDeltasVector11 = VectorDivide(VectorSubtract(DestinationVector11, GainVector11), NumFramesVector);
 
 		VectorRegister4Float GainVector21 = VectorSet(StartGains[4], StartGains[5], StartGains[0], StartGains[1]);
@@ -790,8 +788,8 @@ namespace Audio
 		const VectorRegister4Float GainDeltasVector22 = VectorDivide(VectorSubtract(DestinationVector22, GainVector22), NumFramesVector);
 		GainVector22 = VectorMultiplyAdd(GainDeltasVector22, DeltaHalf21, GainVector22);
 
-		VectorRegister4Float GainVector32 = VectorLoadAligned(StartGains + 8);
-		const VectorRegister4Float DestinationVector32 = VectorLoadAligned(EndGains + 8);
+		VectorRegister4Float GainVector32 = VectorLoad(StartGains + 8);
+		const VectorRegister4Float DestinationVector32 = VectorLoad(EndGains + 8);
 		const VectorRegister4Float GainDeltasVector32 = VectorDivide(VectorSubtract(DestinationVector32, GainVector32), NumFramesVector);
 		GainVector32 = VectorMultiplyAdd(GainDeltasVector32, DeltaHalf31, GainVector32);
 
@@ -810,21 +808,21 @@ namespace Audio
 
 			VectorRegister4Float Result = VectorMultiply(Input11, GainVector11);
 			Result = VectorMultiplyAdd(Input12, GainVector12, Result);
-			VectorStoreAligned(Result, &DestinationBuffer[OutputIndex]);
+			VectorStore(Result, &DestinationBuffer[OutputIndex]);
 
 			GainVector11 = VectorAdd(GainVector11, GainDeltasVector11);
 			GainVector12 = VectorAdd(GainVector12, GainDeltasVector12);
 
 			Result = VectorMultiply(Input21, GainVector21);
 			Result = VectorMultiplyAdd(Input22, GainVector22, Result);
-			VectorStoreAligned(Result, &DestinationBuffer[OutputIndex + 4]);
+			VectorStore(Result, &DestinationBuffer[OutputIndex + 4]);
 
 			GainVector21 = VectorAdd(GainVector21, GainDeltasVector21);
 			GainVector22 = VectorAdd(GainVector22, GainDeltasVector22);
 
 			Result = VectorMultiply(Input31, GainVector31);
 			Result = VectorMultiplyAdd(Input32, GainVector31, Result);
-			VectorStoreAligned(Result, &DestinationBuffer[OutputIndex + 8]);
+			VectorStore(Result, &DestinationBuffer[OutputIndex + 8]);
 
 			GainVector31 = VectorAdd(GainVector31, GainDeltasVector31);
 			GainVector32 = VectorAdd(GainVector32, GainDeltasVector32);
@@ -838,18 +836,18 @@ namespace Audio
 
 	void Apply8ChannelGain(float* RESTRICT InterleavedBuffer, int32 NumSamples, const float* RESTRICT Gains)
 	{
-		const VectorRegister4Float GainVector1 = VectorLoadAligned(Gains);
-		const VectorRegister4Float GainVector2 = VectorLoadAligned(Gains + 4);
+		const VectorRegister4Float GainVector1 = VectorLoad(Gains);
+		const VectorRegister4Float GainVector2 = VectorLoad(Gains + 4);
 
 		for (int32 i = 0; i < NumSamples; i += 8)
 		{
-			VectorRegister4Float Result = VectorLoadAligned(&InterleavedBuffer[i]);
+			VectorRegister4Float Result = VectorLoad(&InterleavedBuffer[i]);
 			Result = VectorMultiply(Result, GainVector1);
-			VectorStoreAligned(Result, &InterleavedBuffer[i]);
+			VectorStore(Result, &InterleavedBuffer[i]);
 
-			Result = VectorLoadAligned(&InterleavedBuffer[i + 4]);
+			Result = VectorLoad(&InterleavedBuffer[i + 4]);
 			Result = VectorMultiply(Result, GainVector2);
-			VectorStoreAligned(Result, &InterleavedBuffer[i + 4]);
+			VectorStore(Result, &InterleavedBuffer[i + 4]);
 		}
 	}
 
@@ -862,25 +860,25 @@ namespace Audio
 	{
 		const VectorRegister4Float NumFramesVector = VectorSetFloat1(NumSamples / 8.0f);
 
-		VectorRegister4Float GainVector1 = VectorLoadAligned(StartGains);
-		const VectorRegister4Float DestinationVector1 = VectorLoadAligned(EndGains);
+		VectorRegister4Float GainVector1 = VectorLoad(StartGains);
+		const VectorRegister4Float DestinationVector1 = VectorLoad(EndGains);
 		const VectorRegister4Float GainDeltasVector1 = VectorDivide(VectorSubtract(DestinationVector1, GainVector1), NumFramesVector);
 
-		VectorRegister4Float GainVector2 = VectorLoadAligned(StartGains + 4);
-		const VectorRegister4Float DestinationVector2 = VectorLoadAligned(EndGains + 4);
+		VectorRegister4Float GainVector2 = VectorLoad(StartGains + 4);
+		const VectorRegister4Float DestinationVector2 = VectorLoad(EndGains + 4);
 		const VectorRegister4Float GainDeltasVector2 = VectorDivide(VectorSubtract(DestinationVector2, GainVector2), NumFramesVector);
 
 		for (int32 i = 0; i < NumSamples; i += 8)
 		{
-			VectorRegister4Float Result = VectorLoadAligned(&InterleavedBuffer[i]);
+			VectorRegister4Float Result = VectorLoad(&InterleavedBuffer[i]);
 			Result = VectorMultiply(Result, GainVector1);
-			VectorStoreAligned(Result, &InterleavedBuffer[i]);
+			VectorStore(Result, &InterleavedBuffer[i]);
 
 			GainVector1 = VectorAdd(GainVector1, GainDeltasVector1);
 
-			Result = VectorLoadAligned(&InterleavedBuffer[i + 4]);
+			Result = VectorLoad(&InterleavedBuffer[i + 4]);
 			Result = VectorMultiply(Result, GainVector2);
-			VectorStoreAligned(Result, &InterleavedBuffer[i + 4]);
+			VectorStore(Result, &InterleavedBuffer[i + 4]);
 
 			GainVector2 = VectorAdd(GainVector2, GainDeltasVector2);
 		}
@@ -907,18 +905,18 @@ namespace Audio
 
 	void MixMonoTo8ChannelsFast(const float* RESTRICT MonoBuffer, float* RESTRICT DestinationBuffer, int32 NumFrames, const float* RESTRICT Gains)
 	{
-		const VectorRegister4Float GainVector1 = VectorLoadAligned(Gains);
-		const VectorRegister4Float GainVector2 = VectorLoadAligned(Gains + 4);
+		const VectorRegister4Float GainVector1 = VectorLoad(Gains);
+		const VectorRegister4Float GainVector2 = VectorLoad(Gains + 4);
 
 		for (int32 i = 0; i < NumFrames; i++)
 		{
 			VectorRegister4Float Result = VectorLoadFloat1(&MonoBuffer[i]);
 			Result = VectorMultiply(Result, GainVector1);
-			VectorStoreAligned(Result, &DestinationBuffer[i * 8]);
+			VectorStore(Result, &DestinationBuffer[i * 8]);
 
 			Result = VectorLoadFloat1(&MonoBuffer[i]);
 			Result = VectorMultiply(Result, GainVector2);
-			VectorStoreAligned(Result, &DestinationBuffer[i * 8 + 4]);
+			VectorStore(Result, &DestinationBuffer[i * 8 + 4]);
 		}
 	}
 
@@ -944,24 +942,24 @@ namespace Audio
 	{
 		const VectorRegister4Float NumFramesVector = VectorSetFloat1((float) NumFrames);
 
-		VectorRegister4Float GainVector1 = VectorLoadAligned(StartGains);
-		const VectorRegister4Float DestinationVector1 = VectorLoadAligned(EndGains);
+		VectorRegister4Float GainVector1 = VectorLoad(StartGains);
+		const VectorRegister4Float DestinationVector1 = VectorLoad(EndGains);
 		const VectorRegister4Float GainDeltasVector1 = VectorDivide(VectorSubtract(DestinationVector1, GainVector1), NumFramesVector);
 
-		VectorRegister4Float GainVector2 = VectorLoadAligned(StartGains + 4);
-		const VectorRegister4Float DestinationVector2 = VectorLoadAligned(EndGains + 4);
+		VectorRegister4Float GainVector2 = VectorLoad(StartGains + 4);
+		const VectorRegister4Float DestinationVector2 = VectorLoad(EndGains + 4);
 		const VectorRegister4Float GainDeltasVector2 = VectorDivide(VectorSubtract(DestinationVector2, GainVector2), NumFramesVector);
 
 		for (int32 i = 0; i < NumFrames; i++)
 		{
 			const VectorRegister4Float Input = VectorLoadFloat1(&MonoBuffer[i]);
 			VectorRegister4Float Result = VectorMultiply(Input, GainVector1);
-			VectorStoreAligned(Result, &DestinationBuffer[i * 8]);
+			VectorStore(Result, &DestinationBuffer[i * 8]);
 
 			GainVector1 = VectorAdd(GainVector1, GainDeltasVector1);
 
 			Result = VectorMultiply(Input, GainVector2);
-			VectorStoreAligned(Result, &DestinationBuffer[i * 8 + 4]);
+			VectorStore(Result, &DestinationBuffer[i * 8 + 4]);
 
 			GainVector2 = VectorAdd(GainVector2, GainDeltasVector2);
 		}
@@ -993,10 +991,10 @@ namespace Audio
 
 	void Mix2ChannelsTo8ChannelsFast(const float* RESTRICT SourceBuffer, float* RESTRICT DestinationBuffer, int32 NumFrames, const float* RESTRICT Gains)
 	{
-		const VectorRegister4Float GainVector11 = VectorLoadAligned(Gains);
-		const VectorRegister4Float GainVector21 = VectorLoadAligned(Gains + 4);
-		const VectorRegister4Float GainVector12 = VectorLoadAligned(Gains + 8);
-		const VectorRegister4Float GainVector22 = VectorLoadAligned(Gains + 12);
+		const VectorRegister4Float GainVector11 = VectorLoad(Gains);
+		const VectorRegister4Float GainVector21 = VectorLoad(Gains + 4);
+		const VectorRegister4Float GainVector12 = VectorLoad(Gains + 8);
+		const VectorRegister4Float GainVector22 = VectorLoad(Gains + 12);
 
 		for (int32 i = 0; i < NumFrames; i++)
 		{
@@ -1006,11 +1004,11 @@ namespace Audio
 			VectorRegister4Float Result = VectorMultiply(Input1, GainVector11);
 			Result = VectorMultiplyAdd(Input2, GainVector12, Result);
 
-			VectorStoreAligned(Result, &DestinationBuffer[i * 8]);
+			VectorStore(Result, &DestinationBuffer[i * 8]);
 
 			Result = VectorMultiply(Input1, GainVector21);
 			Result = VectorMultiplyAdd(Input2, GainVector22, Result);
-			VectorStoreAligned(Result, &DestinationBuffer[i * 8 + 4]);
+			VectorStore(Result, &DestinationBuffer[i * 8 + 4]);
 		}
 	}
 
@@ -1042,20 +1040,20 @@ namespace Audio
 	{
 		const VectorRegister4Float NumFramesVector = VectorSetFloat1((float) NumFrames);
 
-		VectorRegister4Float GainVector11 = VectorLoadAligned(StartGains);
-		const VectorRegister4Float DestinationVector11 = VectorLoadAligned(EndGains);
+		VectorRegister4Float GainVector11 = VectorLoad(StartGains);
+		const VectorRegister4Float DestinationVector11 = VectorLoad(EndGains);
 		const VectorRegister4Float GainDeltasVector11 = VectorDivide(VectorSubtract(DestinationVector11, GainVector11), NumFramesVector);
 
-		VectorRegister4Float GainVector21 = VectorLoadAligned(StartGains + 4);
-		const VectorRegister4Float DestinationVector21 = VectorLoadAligned(EndGains + 4);
+		VectorRegister4Float GainVector21 = VectorLoad(StartGains + 4);
+		const VectorRegister4Float DestinationVector21 = VectorLoad(EndGains + 4);
 		const VectorRegister4Float GainDeltasVector21 = VectorDivide(VectorSubtract(DestinationVector21, GainVector21), NumFramesVector);
 
-		VectorRegister4Float GainVector12 = VectorLoadAligned(StartGains + 8);
-		const VectorRegister4Float DestinationVector12 = VectorLoadAligned(EndGains + 8);
+		VectorRegister4Float GainVector12 = VectorLoad(StartGains + 8);
+		const VectorRegister4Float DestinationVector12 = VectorLoad(EndGains + 8);
 		const VectorRegister4Float GainDeltasVector12 = VectorDivide(VectorSubtract(DestinationVector12, GainVector12), NumFramesVector);
 
-		VectorRegister4Float GainVector22 = VectorLoadAligned(StartGains + 12);
-		const VectorRegister4Float DestinationVector22 = VectorLoadAligned(EndGains + 12);
+		VectorRegister4Float GainVector22 = VectorLoad(StartGains + 12);
+		const VectorRegister4Float DestinationVector22 = VectorLoad(EndGains + 12);
 		const VectorRegister4Float GainDeltasVector22 = VectorDivide(VectorSubtract(DestinationVector22, GainVector22), NumFramesVector);
 
 		for (int32 i = 0; i < NumFrames; i++)
@@ -1065,14 +1063,14 @@ namespace Audio
 
 			VectorRegister4Float Result = VectorMultiply(Input1, GainVector11);
 			Result = VectorMultiplyAdd(Input2, GainVector12, Result);
-			VectorStoreAligned(Result, &DestinationBuffer[i * 8]);
+			VectorStore(Result, &DestinationBuffer[i * 8]);
 
 			GainVector11 = VectorAdd(GainVector11, GainDeltasVector11);
 			GainVector12 = VectorAdd(GainVector12, GainDeltasVector12);
 
 			Result = VectorMultiply(Input1, GainVector21);
 			Result = VectorMultiplyAdd(Input2, GainVector22, Result);
-			VectorStoreAligned(Result, &DestinationBuffer[i * 8 + 4]);
+			VectorStore(Result, &DestinationBuffer[i * 8 + 4]);
 
 			GainVector21 = VectorAdd(GainVector21, GainDeltasVector21);
 			GainVector22 = VectorAdd(GainVector22, GainDeltasVector22);
@@ -1188,9 +1186,6 @@ namespace Audio
 	{
 		checkf(InNum >= 4, TEXT("Buffer must have at least %i elements."), AUDIO_NUM_FLOATS_PER_VECTOR_REGISTER);
 		checkf(0 == (InNum % AUDIO_NUM_FLOATS_PER_VECTOR_REGISTER), TEXT("Buffer length be a multiple of %i."), AUDIO_NUM_FLOATS_PER_VECTOR_REGISTER);
-		checkf(IsAligned<const float*>(InBuffer1, AUDIO_SIMD_BYTE_ALIGNMENT), TEXT("Memory must be aligned to use vector operations."));
-		checkf(IsAligned<const float*>(InBuffer2, AUDIO_SIMD_BYTE_ALIGNMENT), TEXT("Memory must be aligned to use vector operations."));
-		checkf(IsAligned<float*>(OutBuffer, AUDIO_SIMD_BYTE_ALIGNMENT), TEXT("Memory must be aligned to use vector operations."));
 
 		const int32 OutNum = 2 * InNum;
 
@@ -1198,22 +1193,22 @@ namespace Audio
 		for (int32 i = 0; i < InNum; i += AUDIO_NUM_FLOATS_PER_VECTOR_REGISTER)
 		{
 			// Vector1[L0, L1, L2, L3]
-			VectorRegister4Float Vector1 = VectorLoadAligned(&InBuffer1[i]);
+			VectorRegister4Float Vector1 = VectorLoad(&InBuffer1[i]);
 			// Vector2[R0, R1, R2, R3]
-			VectorRegister4Float Vector2 = VectorLoadAligned(&InBuffer2[i]);
+			VectorRegister4Float Vector2 = VectorLoad(&InBuffer2[i]);
 
 			// HalfInterleaved[L0, L1, R0, R1]
 			VectorRegister4Float HalfInterleaved = VectorShuffle(Vector1, Vector2, 0, 1, 0, 1);
 			// Interleaved[L0, R0, L1, R1]
 			VectorRegister4Float Interleaved = VectorSwizzle(HalfInterleaved, 0, 2, 1, 3);
-			VectorStoreAligned(Interleaved, &OutBuffer[OutPos]);
+			VectorStore(Interleaved, &OutBuffer[OutPos]);
 			OutPos += AUDIO_NUM_FLOATS_PER_VECTOR_REGISTER;
 
 			// HalfInterleaved[L2, L3, R2, R3]
 			HalfInterleaved = VectorShuffle(Vector1, Vector2, 2, 3, 2, 3);
 			// Interleaved[L2, R2, L3, R3]
 			Interleaved = VectorSwizzle(HalfInterleaved, 0, 2, 1, 3);
-			VectorStoreAligned(Interleaved, &OutBuffer[OutPos]);
+			VectorStore(Interleaved, &OutBuffer[OutPos]);
 			OutPos += AUDIO_NUM_FLOATS_PER_VECTOR_REGISTER;
 		}
 	}
@@ -1238,25 +1233,22 @@ namespace Audio
 	{
 		checkf(InNumFrames >= AUDIO_NUM_FLOATS_PER_VECTOR_REGISTER, TEXT("Buffer must have at least %i elements."), AUDIO_NUM_FLOATS_PER_VECTOR_REGISTER);
 		checkf(0 == (InNumFrames % AUDIO_NUM_FLOATS_PER_VECTOR_REGISTER), TEXT("Buffer length be a multiple of %i."), AUDIO_NUM_FLOATS_PER_VECTOR_REGISTER);
-		checkf(IsAligned<const float*>(InBuffer, AUDIO_SIMD_BYTE_ALIGNMENT), TEXT("Memory must be aligned to use vector operations."));
-		checkf(IsAligned<const float*>(OutBuffer1, AUDIO_SIMD_BYTE_ALIGNMENT), TEXT("Memory must be aligned to use vector operations."));
-		checkf(IsAligned<float*>(OutBuffer2, AUDIO_SIMD_BYTE_ALIGNMENT), TEXT("Memory must be aligned to use vector operations."));
 
 		int32 InNum = InNumFrames * 2;
 		int32 OutPos = 0;
 		for (int32 InPos = 0; InPos < InNum; InPos += 8)
 		{
 			// load 4 frames (2 frames per vector)
-			VectorRegister4Float InVector1 = VectorLoadAligned(&InBuffer[InPos]);
-			VectorRegister4Float InVector2 = VectorLoadAligned(&InBuffer[InPos + 4]);
+			VectorRegister4Float InVector1 = VectorLoad(&InBuffer[InPos]);
+			VectorRegister4Float InVector2 = VectorLoad(&InBuffer[InPos + 4]);
 
 			// Write channel 0
 			VectorRegister4Float OutVector = VectorShuffle(InVector1, InVector2, 0, 2, 0, 2);
-			VectorStoreAligned(OutVector, &OutBuffer1[OutPos]);
+			VectorStore(OutVector, &OutBuffer1[OutPos]);
 
 			// Write channel 1
 			OutVector = VectorShuffle(InVector1, InVector2, 1, 3, 1, 3);
-			VectorStoreAligned(OutVector, &OutBuffer2[OutPos]);
+			VectorStore(OutVector, &OutBuffer2[OutPos]);
 
 			OutPos += AUDIO_NUM_FLOATS_PER_VECTOR_REGISTER;
 		}
@@ -1279,8 +1271,6 @@ namespace Audio
 	{
 		checkf(InNumFrames >= AUDIO_NUM_FLOATS_PER_VECTOR_REGISTER, TEXT("Buffer must have at least %i elements."), AUDIO_NUM_FLOATS_PER_VECTOR_REGISTER);
 		checkf(0 == (InNumFrames % AUDIO_NUM_FLOATS_PER_VECTOR_REGISTER), TEXT("Buffer length be a multiple of %i."), AUDIO_NUM_FLOATS_PER_VECTOR_REGISTER);
-		checkf(IsAligned<const float*>(InSamples, AUDIO_SIMD_BYTE_ALIGNMENT), TEXT("Memory must be aligned to use vector operations."));
-		checkf(IsAligned<float*>(OutSamples, AUDIO_SIMD_BYTE_ALIGNMENT), TEXT("Memory must be aligned to use vector operations."));
 
 		const int32 InNum = InNumFrames * 2;
 		int32 OutPos = 0;
@@ -1288,9 +1278,9 @@ namespace Audio
 		{
 			// Load 4 frames (2 frames per vector)
 			// Buffer1[L0, R0, L1, R1]
-			VectorRegister4Float Buffer1 = VectorLoadAligned(&InSamples[i]);
+			VectorRegister4Float Buffer1 = VectorLoad(&InSamples[i]);
 			// Buffer2[L2, R2, L3, R3]
-			VectorRegister4Float Buffer2 = VectorLoadAligned(&InSamples[i + 4]);
+			VectorRegister4Float Buffer2 = VectorLoad(&InSamples[i + 4]);
 
 			// Shuffle samples into order
 			// Channel0[L0, L1, L2, L3]
@@ -1302,7 +1292,7 @@ namespace Audio
 			// Out[L0 + R0, L1 + R1, L2 + R2, L3 + R3]
 			VectorRegister4Float Out = VectorAdd(Channel0, Channel1);
 
-			VectorStoreAligned(Out, &OutSamples[OutPos]);
+			VectorStore(Out, &OutSamples[OutPos]);
 			OutPos += AUDIO_NUM_FLOATS_PER_VECTOR_REGISTER;
 		}
 	}
diff --git a/Engine/Source/Runtime/SignalProcessing/Private/Dsp.cpp b/Engine/Source/Runtime/SignalProcessing/Private/Dsp.cpp
--- a/Engine/Source/Runtime/SignalProcessing/Private/Dsp.cpp
+++ b/Engine/Source/Runtime/SignalProcessing/Private/Dsp.cpp
@@ -2,6 +2,8 @@
 
 #include "DSP/Dsp.h"
 
+#include "DSP/FloatArrayMath.h"
+
 namespace Audio
 {
 	void EncodeMidSide(
@@ -62,4 +64,4 @@ namespace Audio
 
 	}
 
-}
\ No newline at end of file
+}
diff --git a/Engine/Source/Runtime/SignalProcessing/Private/DynamicDelayAPF.cpp b/Engine/Source/Runtime/SignalProcessing/Private/DynamicDelayAPF.cpp
--- a/Engine/Source/Runtime/SignalProcessing/Private/DynamicDelayAPF.cpp
+++ b/Engine/Source/Runtime/SignalProcessing/Private/DynamicDelayAPF.cpp
@@ -1,143 +1,143 @@
 // Copyright Epic Games, Inc. All Rights Reserved.
 
 #include "DSP/DynamicDelayAPF.h"
-#include "DSP/BufferVectorOperations.h"
+#include "DSP/FloatArrayMath.h"
 
-using namespace Audio;
-
-FDynamicDelayAPF::FDynamicDelayAPF(float InG, int32 InMinDelay, int32 InMaxDelay, int32 InMaxNumInternalBufferSamples, float InSampleRate)
-:	EaseTimeInSec(1.f)
-,	MinDelay(InMinDelay)
-,	MaxDelay(InMaxDelay)
-,	NumDelaySamples(InMinDelay - 1)
-,	NumInternalBufferSamples(InMaxNumInternalBufferSamples)
+namespace Audio
 {
-	G.Init(InSampleRate);
-	G.SetValue(InG);
-
-	checkf(NumDelaySamples >= 0, TEXT("Minimum delay must be atleast 1"));
-	// NumInternalBufferSamples must be less than the length of the delay to support buffer indexing logic.
-	int32 MaxBufferSamples = FMath::Min(NumDelaySamples, InMaxNumInternalBufferSamples);
-	if (NumInternalBufferSamples > MaxBufferSamples)
+	FDynamicDelayAPF::FDynamicDelayAPF(float InG, int32 InMinDelay, int32 InMaxDelay, int32 InMaxNumInternalBufferSamples, float InSampleRate)
+	:	EaseTimeInSec(1.f)
+	,	MinDelay(InMinDelay)
+	,	MaxDelay(InMaxDelay)
+	,	NumDelaySamples(InMinDelay - 1)
+	,	NumInternalBufferSamples(InMaxNumInternalBufferSamples)
 	{
-		NumInternalBufferSamples = MaxBufferSamples;
-		// Block length must be divisible by simd alignment to support simd operations.
-		NumInternalBufferSamples -= (NumInternalBufferSamples % AUDIO_NUM_FLOATS_PER_VECTOR_REGISTER);
-	}
-
-	checkf(NumInternalBufferSamples > 0, TEXT("Invalid internal buffer length"));
-
-	// Check that delays are sane.
-	checkf(MinDelay <= MaxDelay, TEXT("invalid delay range"));
-
+		G.Init(InSampleRate);
+		G.SetValue(InG);
 
-	// Allocate delay line.
-	IntegerDelayLine = MakeUnique<FAlignedBlockBuffer>((2 * MinDelay) + NumInternalBufferSamples, MinDelay + NumInternalBufferSamples);
-	IntegerDelayLine->AddZeros(NumDelaySamples);
-	//FractionalDelayLine = MakeUnique<FAllPassFractionalDelay>(MaxDelay - MinDelay + 1, NumInternalBufferSamples);
-	FractionalDelayLine = MakeUnique<FLinearInterpFractionalDelay>(MaxDelay - MinDelay + 1, NumInternalBufferSamples);
+		checkf(NumDelaySamples >= 0, TEXT("Minimum delay must be atleast 1"));
+		// NumInternalBufferSamples must be less than the length of the delay to support buffer indexing logic.
+		int32 MaxBufferSamples = FMath::Min(NumDelaySamples, InMaxNumInternalBufferSamples);
+		if (NumInternalBufferSamples > MaxBufferSamples)
+		{
+			NumInternalBufferSamples = MaxBufferSamples;
+			// Block length must be divisible by simd alignment to support simd operations.
+			NumInternalBufferSamples -= (NumInternalBufferSamples % AUDIO_NUM_FLOATS_PER_VECTOR_REGISTER);
+		}
 
-	// Allocate internal buffers.
-	FractionalDelays.Reset(NumInternalBufferSamples);
-	FractionalDelays.AddUninitialized(NumInternalBufferSamples);
-	WorkBufferA.Reset(NumInternalBufferSamples);
-	WorkBufferA.AddUninitialized(NumInternalBufferSamples);
-	WorkBufferB.Reset(NumInternalBufferSamples);
-	WorkBufferB.AddUninitialized(NumInternalBufferSamples);
-}
+		checkf(NumInternalBufferSamples > 0, TEXT("Invalid internal buffer length"));
 
-FDynamicDelayAPF::~FDynamicDelayAPF()
-{}
+		// Check that delays are sane.
+		checkf(MinDelay <= MaxDelay, TEXT("invalid delay range"));
 
-void FDynamicDelayAPF::ProcessAudio(const FAlignedFloatBuffer& InSamples, const FAlignedFloatBuffer& InSampleDelays, FAlignedFloatBuffer& OutSamples)
-{
-	const int32 InNum = InSamples.Num();
-	checkf(InNum == InSampleDelays.Num(), TEXT("InSamples [%d], InSampleDelays [%d] length mismatch"), InNum, InSampleDelays.Num());
 
-	OutSamples.Reset(InNum);
-	OutSamples.AddUninitialized(InNum);
+		// Allocate delay line.
+		IntegerDelayLine = MakeUnique<FAlignedBlockBuffer>((2 * MinDelay) + NumInternalBufferSamples, MinDelay + NumInternalBufferSamples);
+		IntegerDelayLine->AddZeros(NumDelaySamples);
+		//FractionalDelayLine = MakeUnique<FAllPassFractionalDelay>(MaxDelay - MinDelay + 1, NumInternalBufferSamples);
+		FractionalDelayLine = MakeUnique<FLinearInterpFractionalDelay>(MaxDelay - MinDelay + 1, NumInternalBufferSamples);
 
-	if (InNum != InSampleDelays.Num())
-	{
-		// Output silence in the case of an error.
-		FMemory::Memset(OutSamples.GetData(), 0, sizeof(float) * InNum);
-		return;
+		// Allocate internal buffers.
+		FractionalDelays.Reset(NumInternalBufferSamples);
+		FractionalDelays.AddUninitialized(NumInternalBufferSamples);
+		WorkBufferA.Reset(NumInternalBufferSamples);
+		WorkBufferA.AddUninitialized(NumInternalBufferSamples);
+		WorkBufferB.Reset(NumInternalBufferSamples);
+		WorkBufferB.AddUninitialized(NumInternalBufferSamples);
 	}
 
-	const float* InSampleData = InSamples.GetData();
-	float* OutSampleData = OutSamples.GetData();
+	FDynamicDelayAPF::~FDynamicDelayAPF()
+	{}
 
-	int32 LeftOver = InNum;
-	int32 BufferPos = 0;
-	while (LeftOver != 0)
+	void FDynamicDelayAPF::ProcessAudio(const FAlignedFloatBuffer& InSamples, const FAlignedFloatBuffer& InSampleDelays, FAlignedFloatBuffer& OutSamples)
 	{
-		const int32 NumToProcess = FMath::Min(LeftOver, NumInternalBufferSamples);
+		const int32 InNum = InSamples.Num();
+		checkf(InNum == InSampleDelays.Num(), TEXT("InSamples [%d], InSampleDelays [%d] length mismatch"), InNum, InSampleDelays.Num());
 
-		FractionalDelays.Reset(NumToProcess);
-		FractionalDelays.AddUninitialized(NumToProcess);
+		OutSamples.Reset(InNum);
+		OutSamples.AddUninitialized(InNum);
 
-		float* FractionalDelayData = FractionalDelays.GetData();
-		FMemory::Memcpy(FractionalDelayData, &InSampleDelays.GetData()[BufferPos], NumToProcess * sizeof(float));
-		
-		for (int32 i = 0; i < NumToProcess; i++)
+		if (InNum != InSampleDelays.Num())
 		{
-			FractionalDelayData[i] -= NumDelaySamples;
+			// Output silence in the case of an error.
+			FMemory::Memset(OutSamples.GetData(), 0, sizeof(float) * InNum);
+			return;
 		}
-		
-		ProcessAudioBlock(&InSampleData[BufferPos], FractionalDelays, NumToProcess, &OutSampleData[BufferPos]);
 
-		LeftOver -= NumToProcess;
-		BufferPos += NumToProcess;
+		const float* InSampleData = InSamples.GetData();
+		float* OutSampleData = OutSamples.GetData();
+
+		int32 LeftOver = InNum;
+		int32 BufferPos = 0;
+		while (LeftOver != 0)
+		{
+			const int32 NumToProcess = FMath::Min(LeftOver, NumInternalBufferSamples);
+
+			FractionalDelays.Reset(NumToProcess);
+			FractionalDelays.AddUninitialized(NumToProcess);
+
+			float* FractionalDelayData = FractionalDelays.GetData();
+			FMemory::Memcpy(FractionalDelayData, &InSampleDelays.GetData()[BufferPos], NumToProcess * sizeof(float));
+			
+			for (int32 i = 0; i < NumToProcess; i++)
+			{
+				FractionalDelayData[i] -= NumDelaySamples;
+			}
+			
+			ProcessAudioBlock(&InSampleData[BufferPos], FractionalDelays, NumToProcess, &OutSampleData[BufferPos]);
+
+			LeftOver -= NumToProcess;
+			BufferPos += NumToProcess;
+		}
 	}
-}
 
-void FDynamicDelayAPF::ProcessAudioBlock(const float* InSamples, const FAlignedFloatBuffer& InFractionalDelays, const int32 InNum, float* OutSamples)
-{
-	// Make a copy of the delay line w[n - d_int]
-	WorkBufferA.Reset(InNum);
-	WorkBufferA.AddUninitialized(InNum);
-	FMemory::Memcpy(WorkBufferA.GetData(), IntegerDelayLine->InspectSamples(InNum), InNum * sizeof(float));
+	void FDynamicDelayAPF::ProcessAudioBlock(const float* InSamples, const FAlignedFloatBuffer& InFractionalDelays, const int32 InNum, float* OutSamples)
+	{
+		// Make a copy of the delay line w[n - d_int]
+		WorkBufferA.Reset(InNum);
+		WorkBufferA.AddUninitialized(InNum);
+		FMemory::Memcpy(WorkBufferA.GetData(), IntegerDelayLine->InspectSamples(InNum), InNum * sizeof(float));
 
-	// Apply fractional delay
-	// w[n - d]
-	FractionalDelayLine->ProcessAudio(WorkBufferA, InFractionalDelays, WorkBufferB);
+		// Apply fractional delay
+		// w[n - d]
+		FractionalDelayLine->ProcessAudio(WorkBufferA, InFractionalDelays, WorkBufferB);
 
-	DelayLineInput.Reset(InNum);
-	DelayLineInput.AddUninitialized(InNum);
+		DelayLineInput.Reset(InNum);
+		DelayLineInput.AddUninitialized(InNum);
 
-	TArrayView<float> WorkBufferAView(WorkBufferA.GetData(), InNum);
-	TArrayView<float> WorkBufferBView(WorkBufferB.GetData(), InNum);
-	TArrayView<float> DelayLineInputView(DelayLineInput.GetData(), InNum);
-	TArrayView<const float> InSamplesView(InSamples, InNum);
-	TArrayView<float> OutSamplesView(OutSamples, InNum);
+		TArrayView<float> WorkBufferAView(WorkBufferA.GetData(), InNum);
+		TArrayView<float> WorkBufferBView(WorkBufferB.GetData(), InNum);
+		TArrayView<float> DelayLineInputView(DelayLineInput.GetData(), InNum);
+		TArrayView<const float> InSamplesView(InSamples, InNum);
+		TArrayView<float> OutSamplesView(OutSamples, InNum);
 
-	// Get G values to interpolate over
-	const float LastG = G.GetNextValue();
-	const float CurrG = G.GetNextValue(InNum - 1);
+		// Get G values to interpolate over
+		const float LastG = G.GetNextValue();
+		const float CurrG = G.GetNextValue(InNum - 1);
 
-	// WorkBufferA = G * WorkBufferB
-	FMemory::Memcpy(WorkBufferA.GetData(), WorkBufferB.GetData(), InNum * sizeof(float));
-	ArrayFade(WorkBufferAView, LastG, CurrG);
+		// WorkBufferA = G * WorkBufferB
+		FMemory::Memcpy(WorkBufferA.GetData(), WorkBufferB.GetData(), InNum * sizeof(float));
+		ArrayFade(WorkBufferAView, LastG, CurrG);
 
-	// DelayLineInput = InSamples + WorkBufferA
-	// w[n] = x[n] + G * w[n - d]
-	ArraySum(InSamplesView, WorkBufferAView, DelayLineInputView);
+		// DelayLineInput = InSamples + WorkBufferA
+		// w[n] = x[n] + G * w[n - d]
+		ArraySum(InSamplesView, WorkBufferAView, DelayLineInputView);
 
-	ArrayUnderflowClamp(DelayLineInput);
+		ArrayUnderflowClamp(DelayLineInput);
 
-	// Update delay line
-	IntegerDelayLine->RemoveSamples(InNum);
-	IntegerDelayLine->AddSamples(DelayLineInput.GetData(), InNum);
+		// Update delay line
+		IntegerDelayLine->RemoveSamples(InNum);
+		IntegerDelayLine->AddSamples(DelayLineInput.GetData(), InNum);
 
-	// y[n] = -G w[n] + w[n - d]
-	ArrayFade(DelayLineInputView, -LastG, -CurrG);
-	ArraySum(DelayLineInputView, WorkBufferBView, OutSamplesView);
-}
+		// y[n] = -G w[n] + w[n - d]
+		ArrayFade(DelayLineInputView, -LastG, -CurrG);
+		ArraySum(DelayLineInputView, WorkBufferBView, OutSamplesView);
+	}
 
-void FDynamicDelayAPF::Reset() 
-{
-	IntegerDelayLine->ClearSamples();
-	IntegerDelayLine->AddZeros(MinDelay - 1);
-	FractionalDelayLine->Reset();
+	void FDynamicDelayAPF::Reset() 
+	{
+		IntegerDelayLine->ClearSamples();
+		IntegerDelayLine->AddZeros(MinDelay - 1);
+		FractionalDelayLine->Reset();
+	}
 }
-
diff --git a/Engine/Source/Runtime/SignalProcessing/Private/EarlyReflectionsFast.cpp b/Engine/Source/Runtime/SignalProcessing/Private/EarlyReflectionsFast.cpp
--- a/Engine/Source/Runtime/SignalProcessing/Private/EarlyReflectionsFast.cpp
+++ b/Engine/Source/Runtime/SignalProcessing/Private/EarlyReflectionsFast.cpp
@@ -1,237 +1,238 @@
 // Copyright Epic Games, Inc. All Rights Reserved.
 
 #include "DSP/EarlyReflectionsFast.h"
+#include "DSP/FloatArrayMath.h"
 
-using namespace Audio;
-
-namespace EarlyReflectionsPrivate
+namespace Audio
 {
-	float NormalizedLinearToLog(float InLinear)
+	namespace EarlyReflectionsPrivate
 	{
-		check(InLinear >= 0.f);
-		check(InLinear <= 1.f);
-
-		static const float InvLn2 = 1.f / FMath::Loge(2.f);
-		return FMath::Loge(1.f + InLinear) * InvLn2;
+		float NormalizedLinearToLog(float InLinear)
+		{
+			check(InLinear >= 0.f);
+			check(InLinear <= 1.f);
+
+			static const float InvLn2 = 1.f / FMath::Loge(2.f);
+			return FMath::Loge(1.f + InLinear) * InvLn2;
+		}
 	}
-}
 
-FEarlyReflectionsFastSettings::FEarlyReflectionsFastSettings()
-	: Gain(1.0f)
-	, PreDelayMsec(0.0f)
-	, Bandwidth(0.8)
-	, Decay(0.5)
-	, Absorption(0.7)
-{}
-
-bool FEarlyReflectionsFastSettings::operator==(const FEarlyReflectionsFastSettings& Other) const
-{
-	bool bIsEqual = (
-		(Other.Gain == Gain) &&
-		(Other.PreDelayMsec == PreDelayMsec) &&
-		(Other.Bandwidth == Bandwidth) &&
-		(Other.Decay == Decay) &&
-		(Other.Absorption == Absorption));
-	return bIsEqual;
-}
+	FEarlyReflectionsFastSettings::FEarlyReflectionsFastSettings()
+		: Gain(1.0f)
+		, PreDelayMsec(0.0f)
+		, Bandwidth(0.8)
+		, Decay(0.5)
+		, Absorption(0.7)
+	{}
 
-bool FEarlyReflectionsFastSettings::operator!=(const FEarlyReflectionsFastSettings& Other) const
-{
-	return !(*this == Other);
-}
-
-const float FEarlyReflectionsFast::MinGain       = 0.0f;
-const float FEarlyReflectionsFast::MaxGain       = 0.9999f;
-const float FEarlyReflectionsFast::MaxPreDelay   = 1000.0f;
-const float FEarlyReflectionsFast::MinPreDelay   = 0.0f;
-const float FEarlyReflectionsFast::MaxBandwidth  = 0.99999f;
-const float FEarlyReflectionsFast::MinBandwidth  = 0.0f;
-const float FEarlyReflectionsFast::MaxDecay      = 1.0f;
-const float FEarlyReflectionsFast::MinDecay      = 0.0001f;
-const float FEarlyReflectionsFast::MaxAbsorption = 0.99999f;
-const float FEarlyReflectionsFast::MinAbsorption = 0.0f;
-
-const FEarlyReflectionsFastSettings FEarlyReflectionsFast::DefaultSettings;
-
-FEarlyReflectionsFast::FEarlyReflectionsFast(float InSampleRate, int32 InMaxNumInternalBufferSamples, const FEarlyReflectionsFastSettings& InSettings)
-:	Settings(InSettings),
-	SampleRate(InSampleRate),
-	LeftFDN(InMaxNumInternalBufferSamples, FFDNDelaySettings::DefaultLeftDelays(InSampleRate)),
-	RightFDN(InMaxNumInternalBufferSamples, FFDNDelaySettings::DefaultRightDelays(InSampleRate)),
-	LeftPreDelay((int32)InSampleRate * 2.0f, 0, InMaxNumInternalBufferSamples),
-	RightPreDelay((int32)InSampleRate * 2.0f, 0, InMaxNumInternalBufferSamples)
-
-{
-	LeftCoefficients.InputScale = 0.25f;
-	RightCoefficients.InputScale = 0.25f;
-
-	ClampSettings(Settings);
-	ApplySettings();
-}
+	bool FEarlyReflectionsFastSettings::operator==(const FEarlyReflectionsFastSettings& Other) const
+	{
+		bool bIsEqual = (
+			(Other.Gain == Gain) &&
+			(Other.PreDelayMsec == PreDelayMsec) &&
+			(Other.Bandwidth == Bandwidth) &&
+			(Other.Decay == Decay) &&
+			(Other.Absorption == Absorption));
+		return bIsEqual;
+	}
 
-FEarlyReflectionsFast::~FEarlyReflectionsFast()
-{
-}
+	bool FEarlyReflectionsFastSettings::operator!=(const FEarlyReflectionsFastSettings& Other) const
+	{
+		return !(*this == Other);
+	}
 
-void FEarlyReflectionsFast::ClampSettings(FEarlyReflectionsFastSettings& InOutSettings)
-{
-	InOutSettings.Gain = FMath::Clamp(
-		InOutSettings.Gain,
-		FEarlyReflectionsFast::MinGain,
-		FEarlyReflectionsFast::MaxGain);
-	InOutSettings.PreDelayMsec = FMath::Clamp(
-		InOutSettings.PreDelayMsec,
-		FEarlyReflectionsFast::MinPreDelay,
-		FEarlyReflectionsFast::MaxPreDelay);
-	InOutSettings.Bandwidth = FMath::Clamp(
-		InOutSettings.Bandwidth,
-		FEarlyReflectionsFast::MinBandwidth,
-		FEarlyReflectionsFast::MaxBandwidth);
-	InOutSettings.Decay = FMath::Clamp(
-		InOutSettings.Decay,
-		FEarlyReflectionsFast::MinDecay,
-		FEarlyReflectionsFast::MaxDecay);
-	InOutSettings.Absorption = FMath::Clamp(
-		InOutSettings.Absorption,
-		FEarlyReflectionsFast::MinAbsorption,
-		FEarlyReflectionsFast::MaxAbsorption);
-}
+	const float FEarlyReflectionsFast::MinGain       = 0.0f;
+	const float FEarlyReflectionsFast::MaxGain       = 0.9999f;
+	const float FEarlyReflectionsFast::MaxPreDelay   = 1000.0f;
+	const float FEarlyReflectionsFast::MinPreDelay   = 0.0f;
+	const float FEarlyReflectionsFast::MaxBandwidth  = 0.99999f;
+	const float FEarlyReflectionsFast::MinBandwidth  = 0.0f;
+	const float FEarlyReflectionsFast::MaxDecay      = 1.0f;
+	const float FEarlyReflectionsFast::MinDecay      = 0.0001f;
+	const float FEarlyReflectionsFast::MaxAbsorption = 0.99999f;
+	const float FEarlyReflectionsFast::MinAbsorption = 0.0f;
+
+	const FEarlyReflectionsFastSettings FEarlyReflectionsFast::DefaultSettings;
+
+	FEarlyReflectionsFast::FEarlyReflectionsFast(float InSampleRate, int32 InMaxNumInternalBufferSamples, const FEarlyReflectionsFastSettings& InSettings)
+	:	Settings(InSettings),
+		SampleRate(InSampleRate),
+		LeftFDN(InMaxNumInternalBufferSamples, FFDNDelaySettings::DefaultLeftDelays(InSampleRate)),
+		RightFDN(InMaxNumInternalBufferSamples, FFDNDelaySettings::DefaultRightDelays(InSampleRate)),
+		LeftPreDelay((int32)InSampleRate * 2.0f, 0, InMaxNumInternalBufferSamples),
+		RightPreDelay((int32)InSampleRate * 2.0f, 0, InMaxNumInternalBufferSamples)
 
-void FEarlyReflectionsFast::SetSettings(const FEarlyReflectionsFastSettings& InSettings)
-{
-	Settings = InSettings;
-	ClampSettings(Settings);
-	ApplySettings();
-}
+	{
+		LeftCoefficients.InputScale = 0.25f;
+		RightCoefficients.InputScale = 0.25f;
 
-void FEarlyReflectionsFast::ApplySettings()
-{
-	using namespace EarlyReflectionsPrivate;
-
-	int32 DelaySamples = (int32)SampleRate * Settings.PreDelayMsec / 1000.0f;
-	LeftPreDelay.SetDelayLengthSamples(DelaySamples);
-	RightPreDelay.SetDelayLengthSamples(DelaySamples);
-
-	// Convert from linear 0-1 scale to logarithmic 0-1 scale. 
-	const float LPFG = NormalizedLinearToLog(Settings.Bandwidth);
-
-	LeftInputLPF.SetG(LPFG);
-	RightInputLPF.SetG(LPFG);
-
-	LeftCoefficients.LPFB[0] = FMath::Min(Settings.Absorption + 0.1f, 0.9999f);
-	LeftCoefficients.LPFB[1] = FMath::Min(Settings.Absorption - 0.12f, 0.9999f);
-	LeftCoefficients.LPFB[2] = FMath::Min(Settings.Absorption + 0.08f, 0.9999f);
-	LeftCoefficients.LPFB[3] = FMath::Min(Settings.Absorption - 0.07f, 0.9999f);
-	LeftCoefficients.LPFA[0] = 1.0f - LeftCoefficients.LPFB[0];
-	LeftCoefficients.LPFA[1] = 1.0f - LeftCoefficients.LPFB[1];
-	LeftCoefficients.LPFA[2] = 1.0f - LeftCoefficients.LPFB[2];
-	LeftCoefficients.LPFA[3] = 1.0f - LeftCoefficients.LPFB[3];
-	LeftCoefficients.APFG[0] = 0.1f;
-	LeftCoefficients.APFG[1] = 0.2f;
-	LeftCoefficients.APFG[2] = 0.3f;
-	LeftCoefficients.APFG[3] = 0.4f;
-
-	RightCoefficients.LPFB[0] = FMath::Min(Settings.Absorption + 0.17f, 0.999f);
-	RightCoefficients.LPFB[1] = FMath::Min(Settings.Absorption - 0.07f, 0.999f);
-	RightCoefficients.LPFB[2] = FMath::Min(Settings.Absorption + 0.05f, 0.999f);
-	RightCoefficients.LPFB[3] = FMath::Min(Settings.Absorption - 0.11f, 0.999f);
-	RightCoefficients.LPFA[0] = 1.0f - RightCoefficients.LPFB[0];
-	RightCoefficients.LPFA[1] = 1.0f - RightCoefficients.LPFB[1];
-	RightCoefficients.LPFA[2] = 1.0f - RightCoefficients.LPFB[2];
-	RightCoefficients.LPFA[3] = 1.0f - RightCoefficients.LPFB[3];
-	RightCoefficients.APFG[0] = 0.1f;
-	RightCoefficients.APFG[1] = 0.2f;
-	RightCoefficients.APFG[2] = 0.3f;
-	RightCoefficients.APFG[3] = 0.4f;
-
-	// 1/sqrt(2) * Decay
-	//float Feedback = (1.0f - Settings.Decay) * 0.707f;
-	float Feedback = (1.0f - Settings.Decay) * 0.5f;
-	LeftCoefficients.Feedback = Feedback;
-	RightCoefficients.Feedback = Feedback;
-
-	LeftFDN.SetCoefficients(LeftCoefficients);
-	RightFDN.SetCoefficients(RightCoefficients);
-}
+		ClampSettings(Settings);
+		ApplySettings();
+	}
 
-void FEarlyReflectionsFast::ProcessAudio(const FAlignedFloatBuffer& InSamples, const int32 InNumChannels, FAlignedFloatBuffer& OutLeftSamples, FAlignedFloatBuffer& OutRightSamples)
-{
-	checkf((InNumChannels == 1) || (InNumChannels == 2), TEXT("EarlyReflections only supports one or two channel input samples."));
-
-	const int32 InNum = InSamples.Num();
-	const int32 InNumFrames = InNum / InNumChannels;
-
-	// Resize internal buffers
-	LeftInputBuffer.Reset(InNumFrames);
-	LeftWorkBufferA.Reset(InNumFrames);
-	LeftWorkBufferB.Reset(InNumFrames);
-	RightInputBuffer.Reset(InNumFrames);
-	RightWorkBufferA.Reset(InNumFrames);
-	RightWorkBufferB.Reset(InNumFrames);
-
-	LeftInputBuffer.AddUninitialized(InNumFrames);
-	LeftWorkBufferA.AddUninitialized(InNumFrames);
-	LeftWorkBufferB.AddUninitialized(InNumFrames);
-	RightInputBuffer.AddUninitialized(InNumFrames);
-	RightWorkBufferA.AddUninitialized(InNumFrames);
-	RightWorkBufferB.AddUninitialized(InNumFrames);
-
-	// Resize output buffers
-	OutLeftSamples.Reset(InNumFrames);
-	OutRightSamples.Reset(InNumFrames);
-
-	OutLeftSamples.AddUninitialized(InNumFrames);
-	OutRightSamples.AddUninitialized(InNumFrames);
-
-	if (1 == InNumChannels)
+	FEarlyReflectionsFast::~FEarlyReflectionsFast()
 	{
-		// Copy mono to both left and right
-		FMemory::Memcpy(LeftInputBuffer.GetData(), InSamples.GetData(), sizeof(float) * InNumFrames);
-		FMemory::Memcpy(RightInputBuffer.GetData(), InSamples.GetData(), sizeof(float) * InNumFrames);
 	}
-	else if (2 == InNumChannels)
+
+	void FEarlyReflectionsFast::ClampSettings(FEarlyReflectionsFastSettings& InOutSettings)
 	{
-		BufferDeinterleave2ChannelFast(InSamples, LeftInputBuffer, RightInputBuffer);
+		InOutSettings.Gain = FMath::Clamp(
+			InOutSettings.Gain,
+			FEarlyReflectionsFast::MinGain,
+			FEarlyReflectionsFast::MaxGain);
+		InOutSettings.PreDelayMsec = FMath::Clamp(
+			InOutSettings.PreDelayMsec,
+			FEarlyReflectionsFast::MinPreDelay,
+			FEarlyReflectionsFast::MaxPreDelay);
+		InOutSettings.Bandwidth = FMath::Clamp(
+			InOutSettings.Bandwidth,
+			FEarlyReflectionsFast::MinBandwidth,
+			FEarlyReflectionsFast::MaxBandwidth);
+		InOutSettings.Decay = FMath::Clamp(
+			InOutSettings.Decay,
+			FEarlyReflectionsFast::MinDecay,
+			FEarlyReflectionsFast::MaxDecay);
+		InOutSettings.Absorption = FMath::Clamp(
+			InOutSettings.Absorption,
+			FEarlyReflectionsFast::MinAbsorption,
+			FEarlyReflectionsFast::MaxAbsorption);
 	}
-	else 
+
+	void FEarlyReflectionsFast::SetSettings(const FEarlyReflectionsFastSettings& InSettings)
 	{
-		// This class only supports 1 and 2 channel input audio.
-		FMemory::Memset(OutLeftSamples.GetData(), 0, sizeof(float) * InNumFrames);
-		FMemory::Memset(OutRightSamples.GetData(), 0, sizeof(float) * InNumFrames);
-		return;
+		Settings = InSettings;
+		ClampSettings(Settings);
+		ApplySettings();
 	}
 
-	// predelay
-	LeftPreDelay.ProcessAudio(LeftInputBuffer, LeftWorkBufferB);
-	RightPreDelay.ProcessAudio(RightInputBuffer, RightWorkBufferB);
-
-	// lpf
-	LeftInputLPF.ProcessAudio(LeftWorkBufferB, LeftWorkBufferA);
-	RightInputLPF.ProcessAudio(RightWorkBufferB, RightWorkBufferA);
-
-	// feedback delay network
-	LeftFDN.ProcessAudio(LeftWorkBufferA, OutLeftSamples);
-	RightFDN.ProcessAudio(RightWorkBufferA, OutRightSamples);
-	
-	const float OneMinusGain = 1.0 - Settings.Gain;
+	void FEarlyReflectionsFast::ApplySettings()
+	{
+		using namespace EarlyReflectionsPrivate;
+
+		int32 DelaySamples = (int32)SampleRate * Settings.PreDelayMsec / 1000.0f;
+		LeftPreDelay.SetDelayLengthSamples(DelaySamples);
+		RightPreDelay.SetDelayLengthSamples(DelaySamples);
+
+		// Convert from linear 0-1 scale to logarithmic 0-1 scale. 
+		const float LPFG = NormalizedLinearToLog(Settings.Bandwidth);
+
+		LeftInputLPF.SetG(LPFG);
+		RightInputLPF.SetG(LPFG);
+
+		LeftCoefficients.LPFB[0] = FMath::Min(Settings.Absorption + 0.1f, 0.9999f);
+		LeftCoefficients.LPFB[1] = FMath::Min(Settings.Absorption - 0.12f, 0.9999f);
+		LeftCoefficients.LPFB[2] = FMath::Min(Settings.Absorption + 0.08f, 0.9999f);
+		LeftCoefficients.LPFB[3] = FMath::Min(Settings.Absorption - 0.07f, 0.9999f);
+		LeftCoefficients.LPFA[0] = 1.0f - LeftCoefficients.LPFB[0];
+		LeftCoefficients.LPFA[1] = 1.0f - LeftCoefficients.LPFB[1];
+		LeftCoefficients.LPFA[2] = 1.0f - LeftCoefficients.LPFB[2];
+		LeftCoefficients.LPFA[3] = 1.0f - LeftCoefficients.LPFB[3];
+		LeftCoefficients.APFG[0] = 0.1f;
+		LeftCoefficients.APFG[1] = 0.2f;
+		LeftCoefficients.APFG[2] = 0.3f;
+		LeftCoefficients.APFG[3] = 0.4f;
+
+		RightCoefficients.LPFB[0] = FMath::Min(Settings.Absorption + 0.17f, 0.999f);
+		RightCoefficients.LPFB[1] = FMath::Min(Settings.Absorption - 0.07f, 0.999f);
+		RightCoefficients.LPFB[2] = FMath::Min(Settings.Absorption + 0.05f, 0.999f);
+		RightCoefficients.LPFB[3] = FMath::Min(Settings.Absorption - 0.11f, 0.999f);
+		RightCoefficients.LPFA[0] = 1.0f - RightCoefficients.LPFB[0];
+		RightCoefficients.LPFA[1] = 1.0f - RightCoefficients.LPFB[1];
+		RightCoefficients.LPFA[2] = 1.0f - RightCoefficients.LPFB[2];
+		RightCoefficients.LPFA[3] = 1.0f - RightCoefficients.LPFB[3];
+		RightCoefficients.APFG[0] = 0.1f;
+		RightCoefficients.APFG[1] = 0.2f;
+		RightCoefficients.APFG[2] = 0.3f;
+		RightCoefficients.APFG[3] = 0.4f;
+
+		// 1/sqrt(2) * Decay
+		//float Feedback = (1.0f - Settings.Decay) * 0.707f;
+		float Feedback = (1.0f - Settings.Decay) * 0.5f;
+		LeftCoefficients.Feedback = Feedback;
+		RightCoefficients.Feedback = Feedback;
+
+		LeftFDN.SetCoefficients(LeftCoefficients);
+		RightFDN.SetCoefficients(RightCoefficients);
+	}
 
-	// Apply Gain
-	ArrayMultiplyByConstantInPlace(OutLeftSamples, Settings.Gain);
-	ArrayMultiplyByConstantInPlace(OutRightSamples, Settings.Gain);
-}
+	void FEarlyReflectionsFast::ProcessAudio(const FAlignedFloatBuffer& InSamples, const int32 InNumChannels, FAlignedFloatBuffer& OutLeftSamples, FAlignedFloatBuffer& OutRightSamples)
+	{
+		checkf((InNumChannels == 1) || (InNumChannels == 2), TEXT("EarlyReflections only supports one or two channel input samples."));
+
+		const int32 InNum = InSamples.Num();
+		const int32 InNumFrames = InNum / InNumChannels;
+
+		// Resize internal buffers
+		LeftInputBuffer.Reset(InNumFrames);
+		LeftWorkBufferA.Reset(InNumFrames);
+		LeftWorkBufferB.Reset(InNumFrames);
+		RightInputBuffer.Reset(InNumFrames);
+		RightWorkBufferA.Reset(InNumFrames);
+		RightWorkBufferB.Reset(InNumFrames);
+
+		LeftInputBuffer.AddUninitialized(InNumFrames);
+		LeftWorkBufferA.AddUninitialized(InNumFrames);
+		LeftWorkBufferB.AddUninitialized(InNumFrames);
+		RightInputBuffer.AddUninitialized(InNumFrames);
+		RightWorkBufferA.AddUninitialized(InNumFrames);
+		RightWorkBufferB.AddUninitialized(InNumFrames);
+
+		// Resize output buffers
+		OutLeftSamples.Reset(InNumFrames);
+		OutRightSamples.Reset(InNumFrames);
+
+		OutLeftSamples.AddUninitialized(InNumFrames);
+		OutRightSamples.AddUninitialized(InNumFrames);
+
+		if (1 == InNumChannels)
+		{
+			// Copy mono to both left and right
+			FMemory::Memcpy(LeftInputBuffer.GetData(), InSamples.GetData(), sizeof(float) * InNumFrames);
+			FMemory::Memcpy(RightInputBuffer.GetData(), InSamples.GetData(), sizeof(float) * InNumFrames);
+		}
+		else if (2 == InNumChannels)
+		{
+			BufferDeinterleave2ChannelFast(InSamples, LeftInputBuffer, RightInputBuffer);
+		}
+		else 
+		{
+			// This class only supports 1 and 2 channel input audio.
+			FMemory::Memset(OutLeftSamples.GetData(), 0, sizeof(float) * InNumFrames);
+			FMemory::Memset(OutRightSamples.GetData(), 0, sizeof(float) * InNumFrames);
+			return;
+		}
+
+		// predelay
+		LeftPreDelay.ProcessAudio(LeftInputBuffer, LeftWorkBufferB);
+		RightPreDelay.ProcessAudio(RightInputBuffer, RightWorkBufferB);
+
+		// lpf
+		LeftInputLPF.ProcessAudio(LeftWorkBufferB, LeftWorkBufferA);
+		RightInputLPF.ProcessAudio(RightWorkBufferB, RightWorkBufferA);
+
+		// feedback delay network
+		LeftFDN.ProcessAudio(LeftWorkBufferA, OutLeftSamples);
+		RightFDN.ProcessAudio(RightWorkBufferA, OutRightSamples);
+		
+		const float OneMinusGain = 1.0 - Settings.Gain;
+
+		// Apply Gain
+		ArrayMultiplyByConstantInPlace(OutLeftSamples, Settings.Gain);
+		ArrayMultiplyByConstantInPlace(OutRightSamples, Settings.Gain);
+	}
 
-void FEarlyReflectionsFast::FlushAudio()
-{
-	// predelay
-	LeftPreDelay.Reset();
-	RightPreDelay.Reset();
+	void FEarlyReflectionsFast::FlushAudio()
+	{
+		// predelay
+		LeftPreDelay.Reset();
+		RightPreDelay.Reset();
 
-	// lpf
-	LeftInputLPF.FlushAudio();
-	RightInputLPF.FlushAudio();
+		// lpf
+		LeftInputLPF.FlushAudio();
+		RightInputLPF.FlushAudio();
 
-	// feedback delay network
-	LeftFDN.FlushAudio();
-	RightFDN.FlushAudio();
+		// feedback delay network
+		LeftFDN.FlushAudio();
+		RightFDN.FlushAudio();
+	}
 }
-
diff --git a/Engine/Source/Runtime/SignalProcessing/Private/FeedbackDelayNetwork.cpp b/Engine/Source/Runtime/SignalProcessing/Private/FeedbackDelayNetwork.cpp
--- a/Engine/Source/Runtime/SignalProcessing/Private/FeedbackDelayNetwork.cpp
+++ b/Engine/Source/Runtime/SignalProcessing/Private/FeedbackDelayNetwork.cpp
@@ -2,201 +2,200 @@
 
 #include "DSP/FeedbackDelayNetwork.h"
 
-
-using namespace Audio;
-
-// The presets for the late refelections reverb were based on a strange sample rate. This 
-// function helps translate delay values in the paper to delay values for this reverb.
-// https://ccrma.stanford.edu/~dattorro/EffectDesignPart1.pdf
-static int32 FeedbackDelayNetworkGetDelaySamples(float InSampleRate, float InPresetValue) {
-	
-	return (int32)((float)InSampleRate / 29761.f * InPresetValue);
-}
-
-FFDNDelaySettings FFDNDelaySettings::DefaultLeftDelays(float InSampleRate)
-{
-	FFDNDelaySettings Settings;
-	Settings.APF0DelayNumSamples = FeedbackDelayNetworkGetDelaySamples(InSampleRate, 5 * 142);
-	Settings.APF1DelayNumSamples = FeedbackDelayNetworkGetDelaySamples(InSampleRate, 5 * 107);
-	Settings.APF2DelayNumSamples = FeedbackDelayNetworkGetDelaySamples(InSampleRate, 5 * 379);
-	Settings.APF3DelayNumSamples = FeedbackDelayNetworkGetDelaySamples(InSampleRate, 5 * 277);
-	return Settings;
-}
-
-FFDNDelaySettings FFDNDelaySettings::DefaultRightDelays(float InSampleRate)
+namespace Audio
 {
-	FFDNDelaySettings Settings;
-	Settings.APF0DelayNumSamples = FeedbackDelayNetworkGetDelaySamples(InSampleRate, 5 * 279);
-	Settings.APF1DelayNumSamples = FeedbackDelayNetworkGetDelaySamples(InSampleRate, 5 * 137);
-	Settings.APF2DelayNumSamples = FeedbackDelayNetworkGetDelaySamples(InSampleRate, 5 * 213);
-	Settings.APF3DelayNumSamples = FeedbackDelayNetworkGetDelaySamples(InSampleRate, 5 * 327);
-	return Settings;
-}
-
-FFeedbackDelayNetwork::FFeedbackDelayNetwork(int32 InMaxNumInternalBufferSamples, const FFDNDelaySettings& InSettings)
-:	Settings(InSettings)
-{
-	// num internal buffer samples must be less than the minimum delay line and abide by memory alignment.
-	NumInternalBufferSamples = FMath::Min(
-		FMath::Min(Settings.APF0DelayNumSamples, Settings.APF1DelayNumSamples),
-		FMath::Min(Settings.APF2DelayNumSamples, Settings.APF3DelayNumSamples));
-
-	NumInternalBufferSamples = FMath::Min(NumInternalBufferSamples, InMaxNumInternalBufferSamples);
-
-	NumInternalBufferSamples -= (NumInternalBufferSamples % AUDIO_NUM_FLOATS_PER_VECTOR_REGISTER);
+	// The presets for the late refelections reverb were based on a strange sample rate. This 
+	// function helps translate delay values in the paper to delay values for this reverb.
+	// https://ccrma.stanford.edu/~dattorro/EffectDesignPart1.pdf
+	static int32 FeedbackDelayNetworkGetDelaySamples(float InSampleRate, float InPresetValue) {
+		
+		return (int32)((float)InSampleRate / 29761.f * InPresetValue);
+	}
 
-	if (NumInternalBufferSamples <= 0)
+	FFDNDelaySettings FFDNDelaySettings::DefaultLeftDelays(float InSampleRate)
 	{
-		NumInternalBufferSamples = 1;
+		FFDNDelaySettings Settings;
+		Settings.APF0DelayNumSamples = FeedbackDelayNetworkGetDelaySamples(InSampleRate, 5 * 142);
+		Settings.APF1DelayNumSamples = FeedbackDelayNetworkGetDelaySamples(InSampleRate, 5 * 107);
+		Settings.APF2DelayNumSamples = FeedbackDelayNetworkGetDelaySamples(InSampleRate, 5 * 379);
+		Settings.APF3DelayNumSamples = FeedbackDelayNetworkGetDelaySamples(InSampleRate, 5 * 277);
+		return Settings;
 	}
 
-	// Allocate delay lines.
-	DelayLine0 = MakeUnique<FAlignedBlockBuffer>((2 * Settings.APF0DelayNumSamples) + NumInternalBufferSamples, NumInternalBufferSamples);
-	DelayLine1 = MakeUnique<FAlignedBlockBuffer>((2 * Settings.APF1DelayNumSamples) + NumInternalBufferSamples, NumInternalBufferSamples);
-	DelayLine2 = MakeUnique<FAlignedBlockBuffer>((2 * Settings.APF2DelayNumSamples) + NumInternalBufferSamples, NumInternalBufferSamples);
-	DelayLine3 = MakeUnique<FAlignedBlockBuffer>((2 * Settings.APF3DelayNumSamples) + NumInternalBufferSamples, NumInternalBufferSamples);
-
-	// Initialize delay lines
-	DelayLine0->AddZeros(Settings.APF0DelayNumSamples);
-	DelayLine1->AddZeros(Settings.APF1DelayNumSamples);
-	DelayLine2->AddZeros(Settings.APF2DelayNumSamples);
-	DelayLine3->AddZeros(Settings.APF3DelayNumSamples);
-
-	// Initialize internal buffers
-	WorkBuffer0.Reset(NumInternalBufferSamples);
-	WorkBuffer1.Reset(NumInternalBufferSamples);
-	WorkBuffer2.Reset(NumInternalBufferSamples);
-	WorkBuffer3.Reset(NumInternalBufferSamples);
-	WorkBuffer0.AddUninitialized(NumInternalBufferSamples);
-	WorkBuffer1.AddUninitialized(NumInternalBufferSamples);
-	WorkBuffer2.AddUninitialized(NumInternalBufferSamples);
-	WorkBuffer3.AddUninitialized(NumInternalBufferSamples);
-
-	// Initialize filter memory
-	for (int32 i = 0; i < 4; i++)
+	FFDNDelaySettings FFDNDelaySettings::DefaultRightDelays(float InSampleRate)
 	{
-		FMO[i] = 0.0f;
-		LPFZ[i] = 0.0f;
+		FFDNDelaySettings Settings;
+		Settings.APF0DelayNumSamples = FeedbackDelayNetworkGetDelaySamples(InSampleRate, 5 * 279);
+		Settings.APF1DelayNumSamples = FeedbackDelayNetworkGetDelaySamples(InSampleRate, 5 * 137);
+		Settings.APF2DelayNumSamples = FeedbackDelayNetworkGetDelaySamples(InSampleRate, 5 * 213);
+		Settings.APF3DelayNumSamples = FeedbackDelayNetworkGetDelaySamples(InSampleRate, 5 * 327);
+		return Settings;
 	}
-}
 
-FFeedbackDelayNetwork::~FFeedbackDelayNetwork()
-{}
+	FFeedbackDelayNetwork::FFeedbackDelayNetwork(int32 InMaxNumInternalBufferSamples, const FFDNDelaySettings& InSettings)
+	:	Settings(InSettings)
+	{
+		// num internal buffer samples must be less than the minimum delay line and abide by memory alignment.
+		NumInternalBufferSamples = FMath::Min(
+			FMath::Min(Settings.APF0DelayNumSamples, Settings.APF1DelayNumSamples),
+			FMath::Min(Settings.APF2DelayNumSamples, Settings.APF3DelayNumSamples));
 
-void FFeedbackDelayNetwork::SetCoefficients(const FFDNCoefficients& InCoefficients)
-{
-	Coefficients = InCoefficients;
-}
+		NumInternalBufferSamples = FMath::Min(NumInternalBufferSamples, InMaxNumInternalBufferSamples);
 
-void FFeedbackDelayNetwork::ProcessAudio(const FAlignedFloatBuffer& InSamples, FAlignedFloatBuffer& OutSamples)
-{
-	const int32 InNum = InSamples.Num();
-	const float* InSampleData = InSamples.GetData();
+		NumInternalBufferSamples -= (NumInternalBufferSamples % AUDIO_NUM_FLOATS_PER_VECTOR_REGISTER);
 
-	// Prepare output buffer
-	OutSamples.Reset(InNum);
-	OutSamples.AddUninitialized(InNum);
+		if (NumInternalBufferSamples <= 0)
+		{
+			NumInternalBufferSamples = 1;
+		}
+
+		// Allocate delay lines.
+		DelayLine0 = MakeUnique<FAlignedBlockBuffer>((2 * Settings.APF0DelayNumSamples) + NumInternalBufferSamples, NumInternalBufferSamples);
+		DelayLine1 = MakeUnique<FAlignedBlockBuffer>((2 * Settings.APF1DelayNumSamples) + NumInternalBufferSamples, NumInternalBufferSamples);
+		DelayLine2 = MakeUnique<FAlignedBlockBuffer>((2 * Settings.APF2DelayNumSamples) + NumInternalBufferSamples, NumInternalBufferSamples);
+		DelayLine3 = MakeUnique<FAlignedBlockBuffer>((2 * Settings.APF3DelayNumSamples) + NumInternalBufferSamples, NumInternalBufferSamples);
+
+		// Initialize delay lines
+		DelayLine0->AddZeros(Settings.APF0DelayNumSamples);
+		DelayLine1->AddZeros(Settings.APF1DelayNumSamples);
+		DelayLine2->AddZeros(Settings.APF2DelayNumSamples);
+		DelayLine3->AddZeros(Settings.APF3DelayNumSamples);
+
+		// Initialize internal buffers
+		WorkBuffer0.Reset(NumInternalBufferSamples);
+		WorkBuffer1.Reset(NumInternalBufferSamples);
+		WorkBuffer2.Reset(NumInternalBufferSamples);
+		WorkBuffer3.Reset(NumInternalBufferSamples);
+		WorkBuffer0.AddUninitialized(NumInternalBufferSamples);
+		WorkBuffer1.AddUninitialized(NumInternalBufferSamples);
+		WorkBuffer2.AddUninitialized(NumInternalBufferSamples);
+		WorkBuffer3.AddUninitialized(NumInternalBufferSamples);
+
+		// Initialize filter memory
+		for (int32 i = 0; i < 4; i++)
+		{
+			FMO[i] = 0.0f;
+			LPFZ[i] = 0.0f;
+		}
+	}
 
-	// Process audio one block at a time
-	float* OutSampleData = OutSamples.GetData();
-	int32 LeftToProcess = InNum;
-	int32 BufferPos = 0;
+	FFeedbackDelayNetwork::~FFeedbackDelayNetwork()
+	{}
 
-	while (LeftToProcess > 0)
+	void FFeedbackDelayNetwork::SetCoefficients(const FFDNCoefficients& InCoefficients)
 	{
-		int32 NumToProcess = FMath::Min(NumInternalBufferSamples, LeftToProcess);
-		ProcessAudioBuffer(&InSampleData[BufferPos], NumToProcess, &OutSampleData[BufferPos]);
-		LeftToProcess -= NumToProcess;
-		BufferPos += NumToProcess;
+		Coefficients = InCoefficients;
 	}
 
-}
+	void FFeedbackDelayNetwork::ProcessAudio(const FAlignedFloatBuffer& InSamples, FAlignedFloatBuffer& OutSamples)
+	{
+		const int32 InNum = InSamples.Num();
+		const float* InSampleData = InSamples.GetData();
 
-int32 FFeedbackDelayNetwork::GetNumInternalBufferSamples() const
-{
-	return NumInternalBufferSamples;
-}
+		// Prepare output buffer
+		OutSamples.Reset(InNum);
+		OutSamples.AddUninitialized(InNum);
 
-void FFeedbackDelayNetwork::ProcessAudioBuffer(const float* InSampleData, const int32 InNum, float* OutSampleData)
-{
-	const float* DelayLineData0 = DelayLine0->InspectSamples(InNum);
-	const float* DelayLineData1 = DelayLine1->InspectSamples(InNum);
-	const float* DelayLineData2 = DelayLine2->InspectSamples(InNum);
-	const float* DelayLineData3 = DelayLine3->InspectSamples(InNum);
-	float* WorkBuffer0Data = WorkBuffer0.GetData();
-	float* WorkBuffer1Data = WorkBuffer1.GetData();
-	float* WorkBuffer2Data = WorkBuffer2.GetData();
-	float* WorkBuffer3Data = WorkBuffer3.GetData();
-	
-	float AllPassOut[4] = {0.0f};
-	float DelayInput[4] = {0.0f};
-	for (int32 i = 0; i < InNum; i++)
-	{
-		
-		float DelayOut[4] = {DelayLineData0[i], DelayLineData1[i], DelayLineData2[i], DelayLineData3[i]};
+		// Process audio one block at a time
+		float* OutSampleData = OutSamples.GetData();
+		int32 LeftToProcess = InNum;
+		int32 BufferPos = 0;
 
-		// Scale input
-		float ScaledInput = Coefficients.InputScale * InSampleData[i];
-		
-		// Calculate input to delay line. 
-		for (int32 j = 0; j < 4; j++)
+		while (LeftToProcess > 0)
 		{
-			DelayInput[j] = UnderflowClamp((ScaledInput + FMO[j]) + (Coefficients.APFG[j] * DelayOut[j]));
+			int32 NumToProcess = FMath::Min(NumInternalBufferSamples, LeftToProcess);
+			ProcessAudioBuffer(&InSampleData[BufferPos], NumToProcess, &OutSampleData[BufferPos]);
+			LeftToProcess -= NumToProcess;
+			BufferPos += NumToProcess;
 		}
 
-		// Store delay line values
-		WorkBuffer0Data[i] = DelayInput[0];
-		WorkBuffer1Data[i] = DelayInput[1];
-		WorkBuffer2Data[i] = DelayInput[2];
-		WorkBuffer3Data[i] = DelayInput[3];
+	}
+
+	int32 FFeedbackDelayNetwork::GetNumInternalBufferSamples() const
+	{
+		return NumInternalBufferSamples;
+	}
 
-		// complete all pass filters and process low pass filters
-		OutSampleData[i] = 0.0f;
-		for (int32 j = 0; j < 4; j++)
+	void FFeedbackDelayNetwork::ProcessAudioBuffer(const float* InSampleData, const int32 InNum, float* OutSampleData)
+	{
+		const float* DelayLineData0 = DelayLine0->InspectSamples(InNum);
+		const float* DelayLineData1 = DelayLine1->InspectSamples(InNum);
+		const float* DelayLineData2 = DelayLine2->InspectSamples(InNum);
+		const float* DelayLineData3 = DelayLine3->InspectSamples(InNum);
+		float* WorkBuffer0Data = WorkBuffer0.GetData();
+		float* WorkBuffer1Data = WorkBuffer1.GetData();
+		float* WorkBuffer2Data = WorkBuffer2.GetData();
+		float* WorkBuffer3Data = WorkBuffer3.GetData();
+		
+		float AllPassOut[4] = {0.0f};
+		float DelayInput[4] = {0.0f};
+		for (int32 i = 0; i < InNum; i++)
 		{
-			AllPassOut[j] = (DelayInput[j] * -Coefficients.APFG[j]) + DelayOut[j];
-			// Note the reasignment of lpz. 
-			LPFZ[j] = UnderflowClamp((AllPassOut[j] * Coefficients.LPFA[j]) + (LPFZ[j] * Coefficients.LPFB[j]));
-			OutSampleData[i] += LPFZ[j];
-		}
+			
+			float DelayOut[4] = {DelayLineData0[i], DelayLineData1[i], DelayLineData2[i], DelayLineData3[i]};
+
+			// Scale input
+			float ScaledInput = Coefficients.InputScale * InSampleData[i];
+			
+			// Calculate input to delay line. 
+			for (int32 j = 0; j < 4; j++)
+			{
+				DelayInput[j] = UnderflowClamp((ScaledInput + FMO[j]) + (Coefficients.APFG[j] * DelayOut[j]));
+			}
+
+			// Store delay line values
+			WorkBuffer0Data[i] = DelayInput[0];
+			WorkBuffer1Data[i] = DelayInput[1];
+			WorkBuffer2Data[i] = DelayInput[2];
+			WorkBuffer3Data[i] = DelayInput[3];
+
+			// complete all pass filters and process low pass filters
+			OutSampleData[i] = 0.0f;
+			for (int32 j = 0; j < 4; j++)
+			{
+				AllPassOut[j] = (DelayInput[j] * -Coefficients.APFG[j]) + DelayOut[j];
+				// Note the reasignment of lpz. 
+				LPFZ[j] = UnderflowClamp((AllPassOut[j] * Coefficients.LPFA[j]) + (LPFZ[j] * Coefficients.LPFB[j]));
+				OutSampleData[i] += LPFZ[j];
+			}
+
+			// Mix channels and apply feedback.
+			FMO[0] = Coefficients.Feedback * (-LPFZ[1] + LPFZ[2]);
+			FMO[1] = Coefficients.Feedback * (LPFZ[0] + LPFZ[3]);
+			FMO[2] = Coefficients.Feedback * (LPFZ[0] - LPFZ[3]);
+			FMO[3] = Coefficients.Feedback * (-LPFZ[1] - LPFZ[2]);
 
-		// Mix channels and apply feedback.
-		FMO[0] = Coefficients.Feedback * (-LPFZ[1] + LPFZ[2]);
-		FMO[1] = Coefficients.Feedback * (LPFZ[0] + LPFZ[3]);
-		FMO[2] = Coefficients.Feedback * (LPFZ[0] - LPFZ[3]);
-		FMO[3] = Coefficients.Feedback * (-LPFZ[1] - LPFZ[2]);
+		}
 
+		// Update delay lines.
+		DelayLine0->RemoveSamples(InNum);
+		DelayLine1->RemoveSamples(InNum);
+		DelayLine2->RemoveSamples(InNum);
+		DelayLine3->RemoveSamples(InNum);
+		DelayLine0->AddSamples(WorkBuffer0Data, InNum);
+		DelayLine1->AddSamples(WorkBuffer1Data, InNum);
+		DelayLine2->AddSamples(WorkBuffer2Data, InNum);
+		DelayLine3->AddSamples(WorkBuffer3Data, InNum);
 	}
 
-	// Update delay lines.
-	DelayLine0->RemoveSamples(InNum);
-	DelayLine1->RemoveSamples(InNum);
-	DelayLine2->RemoveSamples(InNum);
-	DelayLine3->RemoveSamples(InNum);
-	DelayLine0->AddSamples(WorkBuffer0Data, InNum);
-	DelayLine1->AddSamples(WorkBuffer1Data, InNum);
-	DelayLine2->AddSamples(WorkBuffer2Data, InNum);
-	DelayLine3->AddSamples(WorkBuffer3Data, InNum);
-}
-
-void FFeedbackDelayNetwork::FlushAudio()
-{
-	for (int32 i = 0; i < 4; i++)
+	void FFeedbackDelayNetwork::FlushAudio()
 	{
-		LPFZ[i] = 0.f;
-		FMO[i] = 0.f;
-	}
+		for (int32 i = 0; i < 4; i++)
+		{
+			LPFZ[i] = 0.f;
+			FMO[i] = 0.f;
+		}
 
-	DelayLine0->ClearSamples();
-	DelayLine0->AddZeros(Settings.APF0DelayNumSamples);
+		DelayLine0->ClearSamples();
+		DelayLine0->AddZeros(Settings.APF0DelayNumSamples);
 
-	DelayLine1->ClearSamples();
-	DelayLine1->AddZeros(Settings.APF1DelayNumSamples);
+		DelayLine1->ClearSamples();
+		DelayLine1->AddZeros(Settings.APF1DelayNumSamples);
 
-	DelayLine2->ClearSamples();
-	DelayLine2->AddZeros(Settings.APF2DelayNumSamples);
+		DelayLine2->ClearSamples();
+		DelayLine2->AddZeros(Settings.APF2DelayNumSamples);
 
-	DelayLine3->ClearSamples();
-	DelayLine3->AddZeros(Settings.APF3DelayNumSamples);
+		DelayLine3->ClearSamples();
+		DelayLine3->AddZeros(Settings.APF3DelayNumSamples);
+	}
 }
-
diff --git a/Engine/Source/Runtime/SignalProcessing/Private/FloatArrayMath.cpp b/Engine/Source/Runtime/SignalProcessing/Private/FloatArrayMath.cpp
--- a/Engine/Source/Runtime/SignalProcessing/Private/FloatArrayMath.cpp
+++ b/Engine/Source/Runtime/SignalProcessing/Private/FloatArrayMath.cpp
@@ -70,13 +70,6 @@ namespace Audio
 		}
 	}
 
-	void ArraySum(const FAlignedFloatBuffer& InValues, float& OutSum)
-	{
-		TArrayView<const float> InValuesView(InValues.GetData(), InValues.Num());
-
-		ArraySum(InValuesView, OutSum);
-	}
-
 	void ArraySum(TArrayView<const float> InFloatBuffer1, TArrayView<const float> InFloatBuffer2, TArrayView<float> OutputBuffer)
 	{
 		checkf(InFloatBuffer1.Num() == InFloatBuffer2.Num(), TEXT("Input buffers must be equal length"));
@@ -116,19 +109,6 @@ namespace Audio
 		}
 	}
 
-	void ArraySum(const FAlignedFloatBuffer& InFloatBuffer1, const FAlignedFloatBuffer& InFloatBuffer2, FAlignedFloatBuffer& OutputBuffer)
-	{
-		const int32 InNum = InFloatBuffer1.Num();
-		OutputBuffer.Reset(InNum);
-		OutputBuffer.AddUninitialized(InNum);
-
-		TArrayView<const float> InFloatBuffer1View(InFloatBuffer1.GetData(), InNum);
-		TArrayView<const float> InFloatBuffer2View(InFloatBuffer2.GetData(), InNum);
-		TArrayView<float> OutputBufferView(OutputBuffer.GetData(), InNum);
-
-		ArraySum(InFloatBuffer1View, InFloatBuffer2View, OutputBufferView);
-	}
-
 	void ArrayCumulativeSum(TArrayView<const float> InView, TArray<float>& OutData)
 	{
 		// Initialize output data
@@ -271,12 +251,6 @@ namespace Audio
 		}
 	}
 
-	float ArrayGetMagnitude(const FAlignedFloatBuffer& Buffer)
-	{
-		TArrayView<const float> BufferView(Buffer.GetData(), Buffer.Num());
-		return ArrayGetMagnitude(BufferView);
-	}
-
 	float ArrayGetAverageValue(TArrayView<const float> Buffer)
 	{
 		const int32 Num = Buffer.Num();
@@ -322,18 +296,6 @@ namespace Audio
 		}
 	}
 
-	float ArrayGetAverageValue(const FAlignedFloatBuffer& Buffer)
-	{
-		TArrayView<const float> BufferView(Buffer.GetData(), Buffer.Num());
-		return ArrayGetAverageValue(BufferView);
-	}
-
-	float ArrayGetAverageAbsValue(const FAlignedFloatBuffer& Buffer)
-	{
-		TArrayView<const float> BufferView(Buffer.GetData(), Buffer.Num());
-		return ArrayGetAverageAbsValue(BufferView);
-	}
-
 	float ArrayGetAverageAbsValue(TArrayView<const float> Buffer)
 	{
 		const int32 Num = Buffer.Num();
@@ -784,15 +746,6 @@ namespace Audio
 		}
 	}
 
-	void ArrayMultiplyInPlace(const FAlignedFloatBuffer& InFloatBuffer, FAlignedFloatBuffer& BufferToMultiply)
-	{
-		const int32 InNum = BufferToMultiply.Num();
-		TArrayView<const float> InFloatBufferView(InFloatBuffer.GetData(), InNum);
-		TArrayView<float> BufferToMultiplyView(BufferToMultiply.GetData(), InNum);
-
-		ArrayMultiplyInPlace(InFloatBufferView, BufferToMultiplyView);
-	}
-
 	void ArrayComplexMultiplyInPlace(TArrayView<const float> InValues1, TArrayView<float> InValues2)
 	{
 		check(InValues1.Num() == InValues2.Num());
@@ -850,14 +803,6 @@ namespace Audio
 		}
 	}
 
-	void ArrayComplexMultiplyInPlace(const FAlignedFloatBuffer& InValues1, FAlignedFloatBuffer& InValues2)
-	{
-		TArrayView<const float> Values1View(InValues1.GetData(), InValues1.Num());
-		TArrayView<float> Values2View(InValues2.GetData(), InValues2.Num());
-
-		ArrayComplexMultiplyInPlace(Values1View, Values2View);
-	}
-
 	void ArrayMultiplyByConstant(TArrayView<const float> InFloatBuffer, float InValue, TArrayView<float> OutFloatBuffer)
 	{
 		check(InFloatBuffer.Num() == OutFloatBuffer.Num());
@@ -908,14 +853,6 @@ namespace Audio
 		}
 	}
 
-	void ArrayMultiplyByConstant(const FAlignedFloatBuffer& InFloatBuffer, float InValue, FAlignedFloatBuffer& OutFloatBuffer)
-	{
-		TArrayView<const float> InFloatBufferView(InFloatBuffer.GetData(), InFloatBuffer.Num());
-		TArrayView<float> OutFloatBufferView(OutFloatBuffer.GetData(), OutFloatBuffer.Num());
-
-		ArrayMultiplyByConstant(InFloatBufferView, InValue, OutFloatBufferView);
-	}
-
 	void ArrayMultiplyByConstantInPlace(TArrayView<float> InOutBuffer, float InGain)
 	{
 		int32 Num = InOutBuffer.Num();
@@ -954,13 +891,6 @@ namespace Audio
 		}
 	}
 
-	void ArrayMultiplyByConstantInPlace(FAlignedFloatBuffer& InOutBuffer, float InGain)
-	{
-		TArrayView<float> InOutBufferView(InOutBuffer.GetData(), InOutBuffer.Num());
-
-		ArrayMultiplyByConstantInPlace(InOutBufferView, InGain);
-	}
-
 	void ArrayAddInPlace(TArrayView<const float> InValues, TArrayView<float> InAccumulateValues)
 	{
 		check(InValues.Num() == InAccumulateValues.Num());
@@ -1003,14 +933,6 @@ namespace Audio
 		}
 	}
 
-	void ArrayAddInPlace(const FAlignedFloatBuffer& InValues, FAlignedFloatBuffer& InAccumulateValues)
-	{
-		TArrayView<const float> ValuesView(InValues.GetData(), InValues.Num());
-		TArrayView<float> AccumulateView(InAccumulateValues.GetData(), InAccumulateValues.Num());
-
-		ArrayAddInPlace(ValuesView, AccumulateView);
-	}
-
 	void ArrayAddConstantInplace(TArrayView<float> InOutBuffer, float InConstant)
 	{
 		int32 Num = InOutBuffer.Num();
@@ -1049,14 +971,6 @@ namespace Audio
 		}
 	}
 
-	// Adds a constant to a buffer (useful for DC offset removal)
-	void ArrayAddConstantInplace(FAlignedFloatBuffer& InOutBuffer, float InConstant)
-	{
-		TArrayView<float> InOutBufferView(InOutBuffer.GetData(), InOutBuffer.Num());
-
-		ArrayAddConstantInplace(InOutBufferView, InConstant);
-	}
-
 	void ArrayMultiplyAddInPlace(TArrayView<const float> InValues, float InMultiplier, TArrayView<float> InAccumulateValues)
 	{
 		check(InValues.Num() == InAccumulateValues.Num());
@@ -1081,14 +995,6 @@ namespace Audio
 		}
 	}
 
-	void ArrayMultiplyAddInPlace(const FAlignedFloatBuffer& InValues, float InMultiplier, FAlignedFloatBuffer& InAccumulateValues)
-	{
-		TArrayView<const float> InValuesView(InValues.GetData(), InValues.Num());
-		TArrayView<float> InAccumulateValuesView(InAccumulateValues.GetData(), InAccumulateValues.Num());
-
-		ArrayMultiplyAddInPlace(InValuesView, InMultiplier, InAccumulateValuesView);
-	}
-
 	void ArrayLerpAddInPlace(TArrayView<const float> InValues, float InStartMultiplier, float InEndMultiplier, TArrayView<float> InAccumulateValues)
 	{
 		check(InValues.Num() == InAccumulateValues.Num());
@@ -1142,14 +1048,6 @@ namespace Audio
 		}
 	}
 
-	void ArrayLerpAddInPlace(const FAlignedFloatBuffer& InValues, float InStartMultiplier, float InEndMultiplier, FAlignedFloatBuffer& InAccumulateValues)
-	{
-		TArrayView<const float> ValuesView(InValues.GetData(), InValues.Num());
-		TArrayView<float> AccumulateView(InAccumulateValues.GetData(), InAccumulateValues.Num());
-
-		ArrayLerpAddInPlace(ValuesView, InStartMultiplier, InEndMultiplier, AccumulateView);
-	}
-
 	/* Subtracts two buffers together element-wise. */
 	void ArraySubtract(TArrayView<const float> InMinuend, TArrayView<const float> InSubtrahend, TArrayView<float> OutBuffer)
 	{
@@ -1189,20 +1087,6 @@ namespace Audio
 		}
 	}
 
-	/* Subtracts two buffers together element-wise. */
-	void ArraySubtract(const FAlignedFloatBuffer& InMinuend, const FAlignedFloatBuffer& InSubtrahend, FAlignedFloatBuffer& OutputBuffer)
-	{
-		const int32 InNum = InMinuend.Num();
-		OutputBuffer.Reset(InNum);
-		OutputBuffer.AddUninitialized(InNum);
-
-		TArrayView<const float> InMinuendView(InMinuend.GetData(), InNum);
-		TArrayView<const float> InSubtrahendView(InSubtrahend.GetData(), InNum);
-		TArrayView<float> OutputBufferView(OutputBuffer.GetData(), InNum);
-
-		ArraySubtract(InMinuendView, InSubtrahendView, OutputBufferView);
-	}
-
 	/* Performs element-wise in-place subtraction placing the result in the subtrahend. InOutSubtrahend = InMinuend - InOutSubtrahend */
 	void ArraySubtractInPlace1(TArrayView<const float> InMinuend, TArrayView<float> InOutSubtrahend)
 	{
@@ -1243,17 +1127,6 @@ namespace Audio
 		}
 	}
 
-	/* Performs element-wise in-place subtraction placing the result in the subtrahend. InOutSubtrahend = InMinuend - InOutSubtrahend */
-	void ArraySubtractInPlace1(const FAlignedFloatBuffer& InMinuend, FAlignedFloatBuffer& InOutSubtrahend)
-	{
-		const int32 InNum = InMinuend.Num();
-
-		TArrayView<const float> InMinuendView(InMinuend.GetData(), InNum);
-		TArrayView<float> InOutSubtrahendView(InOutSubtrahend.GetData(), InNum);
-
-		ArraySubtractInPlace1(InMinuendView, InOutSubtrahendView);
-	}
-
 	/* Performs element-wise in-place subtraction placing the result in the minuend. InOutMinuend = InOutMinuend - InSubtrahend */
 	void ArraySubtractInPlace2(TArrayView<float> InOutMinuend, TArrayView<const float> InSubtrahend)
 	{
@@ -1294,17 +1167,6 @@ namespace Audio
 		}
 	}
 
-	/* Performs element-wise in-place subtraction placing the result in the minuend. InOutMinuend = InOutMinuend - InSubtrahend */
-	void ArraySubtractInPlace2(FAlignedFloatBuffer& InOutMinuend, const FAlignedFloatBuffer& InSubtrahend)
-	{
-		const int32 InNum = InOutMinuend.Num();
-
-		TArrayView<float> InOutMinuendView(InOutMinuend.GetData(), InNum);
-		TArrayView<const float> InSubtrahendView(InSubtrahend.GetData(), InNum);
-
-		ArraySubtractInPlace2(InOutMinuendView, InSubtrahendView);
-	}
-
 	void ArraySubtractByConstantInPlace(TArrayView<float> InValues, float InSubtrahend)
 	{
 		const int32 Num = InValues.Num();
@@ -1342,13 +1204,6 @@ namespace Audio
 		}
 	}
 
-	void ArraySubtractByConstantInPlace(FAlignedFloatBuffer& InValues, float InSubtrahend)
-	{
-		TArrayView<float> View(InValues.GetData(), InValues.Num());
-
-		ArraySubtractByConstantInPlace(View, InSubtrahend);
-	}
-
 	void ArraySquare(TArrayView<const float> InValues, TArrayView<float> OutValues)
 	{
 		check(InValues.Num() == OutValues.Num());
@@ -1492,14 +1347,6 @@ namespace Audio
 		}
 	}
 
-	void ArrayComplexConjugate(const FAlignedFloatBuffer& InValues, FAlignedFloatBuffer& OutValues)
-	{
-		TArrayView<const float> InView(InValues.GetData(), InValues.Num());
-		TArrayView<float> OutView(OutValues.GetData(), OutValues.Num());
-
-		ArrayComplexConjugate(InView, OutView);
-	}
-
 	void ArrayComplexConjugateInPlace(TArrayView<float> InValues)
 	{
 		check((InValues.Num() % 2) == 0);
@@ -1543,13 +1390,6 @@ namespace Audio
 		}
 	}
 
-	void ArrayComplexConjugateInPlace(FAlignedFloatBuffer& InValues)
-	{
-		TArrayView<float> InView(InValues.GetData(), InValues.Num());
-
-		ArrayComplexConjugateInPlace(InView);
-	}
-
 	void ArrayMagnitudeToDecibelInPlace(TArrayView<float> InValues, float InMinimumDb)
 	{
 		const int32 Num = InValues.Num();
@@ -1597,13 +1437,6 @@ namespace Audio
 		}
 	}
 
-	void ArrayMagnitudeToDecibelInPlace(FAlignedFloatBuffer& InValues, float InMinimumDb)
-	{
-		TArrayView<float> InView(InValues.GetData(), InValues.Num());
-
-		ArrayMagnitudeToDecibelInPlace(InView, InMinimumDb);
-	}
-
 	void ArrayPowerToDecibelInPlace(TArrayView<float> InValues, float InMinimumDb)
 	{
 		const int32 Num = InValues.Num();
@@ -1651,13 +1484,6 @@ namespace Audio
 		}
 	}
 
-	void ArrayPowerToDecibelInPlace(FAlignedFloatBuffer& InValues, float InMinimumDb)
-	{
-		TArrayView<float> InView(InValues.GetData(), InValues.Num());
-
-		ArrayPowerToDecibelInPlace(InView, InMinimumDb);
-	}
-
 	void ArrayComplexToPower(TArrayView<const float> InComplexValues, TArrayView<float> OutPowerValues)
 	{
 		check((InComplexValues.Num() % 2) == 0);
@@ -1713,14 +1539,6 @@ namespace Audio
 		}
 	}
 
-	void ArrayComplexToPower(const FAlignedFloatBuffer& InComplexValues, FAlignedFloatBuffer& OutPowerValues)
-	{
-		TArrayView<const float> ComplexView(InComplexValues.GetData(), InComplexValues.Num());
-		TArrayView<float> PowerView(OutPowerValues.GetData(), OutPowerValues.Num());
-
-		ArrayComplexToPower(ComplexView, PowerView);
-	}
-
 	void ArrayComplexToPower(TArrayView<const float> InRealSamples, TArrayView<const float> InImaginarySamples, TArrayView<float> OutPowerSamples)
 	{
 		checkf(InRealSamples.Num() == InImaginarySamples.Num(), TEXT("Input buffers must have equal number of elements"));
@@ -1765,20 +1583,6 @@ namespace Audio
 		}
 	}
 
-	void ArrayComplexToPower(const FAlignedFloatBuffer& InRealSamples, const FAlignedFloatBuffer& InImaginarySamples, FAlignedFloatBuffer& OutPowerSamples)
-	{
-		const int32 Num = InRealSamples.Num();
-
-		OutPowerSamples.Reset(Num);
-		OutPowerSamples.AddUninitialized(Num);
-
-		TArrayView<const float> InRealSamplesView(InRealSamples.GetData(), Num);
-		TArrayView<const float> InImaginarySamplesView(InImaginarySamples.GetData(), Num);
-		TArrayView<float> OutPowerSamplesView(OutPowerSamples.GetData(), Num);
-
-		ArrayComplexToPower(InRealSamplesView, InImaginarySamplesView, OutPowerSamplesView);
-	}
-
 	/* Sets a values to zero if value is denormal. Denormal numbers significantly slow down floating point operations. */
 	void ArrayUnderflowClamp(TArrayView<float> InOutValues)
 	{
@@ -1831,13 +1635,6 @@ namespace Audio
 		}
 	}
 
-	void ArrayUnderflowClamp(FAlignedFloatBuffer& InOutBuffer)
-	{
-		TArrayView<float> InOutBufferView(InOutBuffer.GetData(), InOutBuffer.Num());
-
-		ArrayUnderflowClamp(InOutBufferView);
-	}
-
 	/* Clamps values in the buffer to be between InMinValue and InMaxValue */
 	void ArrayRangeClamp(TArrayView<float> InOutBuffer, float InMinValue, float InMaxValue)
 	{
@@ -1886,13 +1683,6 @@ namespace Audio
 		}
 	}
 
-	void ArrayRangeClamp(FAlignedFloatBuffer& InOutBuffer, float InMinValue, float InMaxValue)
-	{
-		TArrayView<float> InOutBufferView(InOutBuffer.GetData(), InOutBuffer.Num());
-
-		ArrayRangeClamp(InOutBufferView, InMinValue, InMaxValue);
-	}
-
 	void ArraySetToConstantInplace(TArrayView<float> InOutBuffer, float InConstant)
 	{
 		int32 Num = InOutBuffer.Num();
@@ -1929,13 +1719,6 @@ namespace Audio
 		}
 	}
 
-	void ArraySetToConstantInplace(FAlignedFloatBuffer& InOutBuffer, float InConstant)
-	{
-		TArrayView<float> InOutBufferView(InOutBuffer.GetData(), InOutBuffer.Num());
-
-		ArraySetToConstantInplace(InOutBufferView, InConstant);
-	}
-
 	/* Performs an element-wise weighted sum OutputBuffer = (InBuffer1 x InGain1) + (InBuffer2 x InGain2) */
 	void ArrayWeightedSum(TArrayView<const float> InBuffer1, float InGain1, TArrayView<const float> InBuffer2, float InGain2, TArrayView<float> OutBuffer)
 	{
@@ -1986,19 +1769,6 @@ namespace Audio
 		}
 	}
 
-	/* Performs an element-wise weighted sum OutputBuffer = (InBuffer1 x InGain1) + (InBuffer2 x InGain2) */
-	void ArrayWeightedSum(const FAlignedFloatBuffer& InBuffer1, float InGain1, const FAlignedFloatBuffer& InBuffer2, float InGain2, FAlignedFloatBuffer& OutBuffer)
-	{
-		OutBuffer.Reset();
-		OutBuffer.AddUninitialized(InBuffer1.Num());
-
-		TArrayView<const float> InBuffer1View(InBuffer1.GetData(), InBuffer1.Num());
-		TArrayView<const float> InBuffer2View(InBuffer2.GetData(), InBuffer2.Num());
-		TArrayView<float> OutBufferView(OutBuffer.GetData(), OutBuffer.Num());
-
-		ArrayWeightedSum(InBuffer1View, InGain1, InBuffer2View, InGain2, OutBufferView);
-	}
-
 	/* Performs an element-wise weighted sum OutputBuffer = (InBuffer1 x InGain1) + InBuffer2 */
 	void ArrayWeightedSum(TArrayView<const float> InBuffer1, float InGain1, TArrayView<const float> InBuffer2, TArrayView<float> OutBuffer)
 	{
@@ -2044,19 +1814,6 @@ namespace Audio
 		}
 	}
 
-	/* Performs an element-wise weighted sum OutputBuffer = (InBuffer1 x InGain1) + InBuffer2 */
-	void ArrayWeightedSum(const FAlignedFloatBuffer& InBuffer1, float InGain1, const FAlignedFloatBuffer& InBuffer2, FAlignedFloatBuffer& OutBuffer)
-	{
-		OutBuffer.Reset();
-		OutBuffer.AddUninitialized(InBuffer1.Num());
-
-		TArrayView<const float> InBuffer1View(InBuffer1.GetData(), InBuffer1.Num());
-		TArrayView<const float> InBuffer2View(InBuffer2.GetData(), InBuffer2.Num());
-		TArrayView<float> OutBufferView(OutBuffer.GetData(), OutBuffer.Num());
-
-		ArrayWeightedSum(InBuffer1View, InGain1, InBuffer2View, OutBufferView);
-	}
-
 	void ArrayFade(TArrayView<float> InOutBuffer, const float StartValue, const float EndValue)
 	{
 		int32 Num = InOutBuffer.Num();
@@ -2120,13 +1877,6 @@ namespace Audio
 		}
 	}
 
-	void ArrayFade(FAlignedFloatBuffer& OutFloatBuffer, const float StartValue, const float EndValue)
-	{
-		TArrayView<float> OutFloatBufferView(OutFloatBuffer.GetData(), OutFloatBuffer.Num());
-
-		ArrayFade(OutFloatBufferView, StartValue, EndValue);
-	}
-
 	void ArrayMixIn(TArrayView<const float> InFloatBuffer, TArrayView<float> BufferToSumTo, const float Gain)
 	{
 		checkf(InFloatBuffer.Num() == BufferToSumTo.Num(), TEXT("Buffers must be equal size"));
@@ -2169,14 +1919,6 @@ namespace Audio
 		}
 	}
 
-	void ArrayMixIn(const FAlignedFloatBuffer& InFloatBuffer, FAlignedFloatBuffer& BufferToSumTo, const float Gain)
-	{
-		TArrayView<const float> InFloatBufferView(InFloatBuffer.GetData(), InFloatBuffer.Num());
-		TArrayView<float> BufferToSumToView(BufferToSumTo.GetData(), BufferToSumTo.Num());
-
-		ArrayMixIn(InFloatBufferView, BufferToSumToView, Gain);
-	}
-
 	void ArrayMixIn(TArrayView<const float> InFloatBuffer, TArrayView<float> BufferToSumTo)
 	{
 		checkf(InFloatBuffer.Num() == BufferToSumTo.Num(), TEXT("Buffers must be equal size"));
@@ -2217,14 +1959,6 @@ namespace Audio
 		}
 	}
 
-	void ArrayMixIn(const FAlignedFloatBuffer& InFloatBuffer, FAlignedFloatBuffer& BufferToSumTo)
-	{
-		TArrayView<const float> InFloatBufferView(InFloatBuffer.GetData(), InFloatBuffer.Num());
-		TArrayView<float> BufferToSumToView(BufferToSumTo.GetData(), BufferToSumTo.Num());
-
-		ArrayMixIn(InFloatBufferView, BufferToSumToView);
-	}
-
 	void ArrayMixIn(TArrayView<const float> InFloatBuffer, TArrayView<float> BufferToSumTo, const float StartGain, const float EndGain)
 	{
 		checkf(InFloatBuffer.Num() == BufferToSumTo.Num(), TEXT("Buffers must be equal size"));
@@ -2291,14 +2025,6 @@ namespace Audio
 		}
 	}
 
-	void ArrayMixIn(const FAlignedFloatBuffer& InFloatBuffer, FAlignedFloatBuffer& BufferToSumTo, const float StartGain, const float EndGain)
-	{
-		TArrayView<const float> InFloatBufferView(InFloatBuffer.GetData(), InFloatBuffer.Num());
-		TArrayView<float> BufferToSumToView(BufferToSumTo.GetData(), BufferToSumTo.Num());
-
-		ArrayMixIn(InFloatBufferView, BufferToSumToView, StartGain, EndGain);
-	}
-
 	void ArrayFloatToPcm16(TArrayView<const float> InView, TArrayView<int16> OutView)
 	{
 		check(InView.Num() == OutView.Num());
diff --git a/Engine/Source/Runtime/SignalProcessing/Private/IntegerDelay.cpp b/Engine/Source/Runtime/SignalProcessing/Private/IntegerDelay.cpp
--- a/Engine/Source/Runtime/SignalProcessing/Private/IntegerDelay.cpp
+++ b/Engine/Source/Runtime/SignalProcessing/Private/IntegerDelay.cpp
@@ -3,119 +3,119 @@
 #include "DSP/IntegerDelay.h"
 #include "DSP/BufferVectorOperations.h"
 
-using namespace Audio;
-
-FIntegerDelay::FIntegerDelay(int32 InMaxNumDelaySamplesSamples, int32 InNumDelaySamples, int32 InNumInternalBufferSamples)
-:	MaxNumDelaySamples(InMaxNumDelaySamplesSamples)
-	, NumDelaySamples(0)
-	, NumDelayLineOffsetSamples(0)
-	, NumBufferOffsetSamples(0)
-	, NumInternalBufferSamples(InNumInternalBufferSamples)
+namespace Audio
 {
-	// Allocate and prepare delay line for maximum delay.
-	DelayLine = MakeUnique<FAlignedBlockBuffer>((2 * MaxNumDelaySamples) + NumInternalBufferSamples, MaxNumDelaySamples + NumInternalBufferSamples);
-	DelayLine->AddZeros(MaxNumDelaySamples);
-
-	// Set current delay.
-	SetDelayLengthSamples(InNumDelaySamples);
-}
+	FIntegerDelay::FIntegerDelay(int32 InMaxNumDelaySamplesSamples, int32 InNumDelaySamples, int32 InNumInternalBufferSamples)
+	:	MaxNumDelaySamples(InMaxNumDelaySamplesSamples)
+		, NumDelaySamples(0)
+		, NumDelayLineOffsetSamples(0)
+		, NumBufferOffsetSamples(0)
+		, NumInternalBufferSamples(InNumInternalBufferSamples)
+	{
+		// Allocate and prepare delay line for maximum delay.
+		DelayLine = MakeUnique<FAlignedBlockBuffer>((2 * MaxNumDelaySamples) + NumInternalBufferSamples, MaxNumDelaySamples + NumInternalBufferSamples);
+		DelayLine->AddZeros(MaxNumDelaySamples);
 
+		// Set current delay.
+		SetDelayLengthSamples(InNumDelaySamples);
+	}
 
-// Destructor
-FIntegerDelay::~FIntegerDelay()
-{}
 
-void FIntegerDelay::SetDelayLengthSamples(int32 InNumDelaySamples)
-{
-	checkf(InNumDelaySamples <= MaxNumDelaySamples, TEXT("InNumDelaySamples must be less than or equal to MaxNumDelaySamples"));
-	checkf(InNumDelaySamples >= 0, TEXT("InNumDelaySamples must be greater than or equal to 0"));
+	// Destructor
+	FIntegerDelay::~FIntegerDelay()
+	{}
 
-	if (InNumDelaySamples > MaxNumDelaySamples)
+	void FIntegerDelay::SetDelayLengthSamples(int32 InNumDelaySamples)
 	{
-		InNumDelaySamples = MaxNumDelaySamples;
+		checkf(InNumDelaySamples <= MaxNumDelaySamples, TEXT("InNumDelaySamples must be less than or equal to MaxNumDelaySamples"));
+		checkf(InNumDelaySamples >= 0, TEXT("InNumDelaySamples must be greater than or equal to 0"));
+
+		if (InNumDelaySamples > MaxNumDelaySamples)
+		{
+			InNumDelaySamples = MaxNumDelaySamples;
+		}
+		else if (InNumDelaySamples < 0)
+		{
+			InNumDelaySamples = 0;
+		}
+
+		NumDelaySamples = InNumDelaySamples;
+		// Store offset location for reading samples out from delay line.
+		// We keep a couple offsets around to ensure we do go over buffer bounds
+		// and abide by alignment rules of FAlignedBlockBuffer
+		NumDelayLineOffsetSamples = MaxNumDelaySamples - NumDelaySamples;
+		NumBufferOffsetSamples = 0;
+		while (0 != (NumDelayLineOffsetSamples % AUDIO_NUM_FLOATS_PER_VECTOR_REGISTER))
+		{
+			NumDelayLineOffsetSamples--;
+			NumBufferOffsetSamples++;
+		}
 	}
-	else if (InNumDelaySamples < 0)
+
+	// Resets the delay line state, flushes buffer and resets read/write pointers.
+	void FIntegerDelay::Reset() 
 	{
-		InNumDelaySamples = 0;
+		DelayLine->ClearSamples();
+		DelayLine->AddZeros(MaxNumDelaySamples);
 	}
 
-	NumDelaySamples = InNumDelaySamples;
-	// Store offset location for reading samples out from delay line.
-	// We keep a couple offsets around to ensure we do go over buffer bounds
-	// and abide by alignment rules of FAlignedBlockBuffer
-	NumDelayLineOffsetSamples = MaxNumDelaySamples - NumDelaySamples;
-	NumBufferOffsetSamples = 0;
-	while (0 != (NumDelayLineOffsetSamples % AUDIO_NUM_FLOATS_PER_VECTOR_REGISTER))
+	// Returns the current delay line length (in samples).
+	int32 FIntegerDelay::GetNumDelaySamples() const 
 	{
-		NumDelayLineOffsetSamples--;
-		NumBufferOffsetSamples++;
+		return NumDelaySamples;
 	}
-}
-
-// Resets the delay line state, flushes buffer and resets read/write pointers.
-void FIntegerDelay::Reset() 
-{
-	DelayLine->ClearSamples();
-	DelayLine->AddZeros(MaxNumDelaySamples);
-}
 
-// Returns the current delay line length (in samples).
-int32 FIntegerDelay::GetNumDelaySamples() const 
-{
-	return NumDelaySamples;
-}
-
-void FIntegerDelay::ProcessAudio(const Audio::FAlignedFloatBuffer& InSamples, Audio::FAlignedFloatBuffer& OutSamples)
-{
-	const float* InSampleData = InSamples.GetData();
-	const int32 InNum = InSamples.Num();
-
-	// Prepare output buffer
-	OutSamples.Reset(InNum);
-	OutSamples.AddUninitialized(InNum);
-	float* OutSampleData = OutSamples.GetData();
-
-	// Process audio one block at a time.
-	int32 LeftOver = InNum;
-	int32 BufferPos = 0;
-	while (LeftOver > 0)
+	void FIntegerDelay::ProcessAudio(const Audio::FAlignedFloatBuffer& InSamples, Audio::FAlignedFloatBuffer& OutSamples)
 	{
-		int32 NumToProcess = FMath::Min(LeftOver, NumInternalBufferSamples);
-		ProcessAudioBlock(&InSampleData[BufferPos], NumToProcess, &OutSampleData[BufferPos]);
-		BufferPos += NumToProcess;
-		LeftOver -= NumToProcess;
+		const float* InSampleData = InSamples.GetData();
+		const int32 InNum = InSamples.Num();
+
+		// Prepare output buffer
+		OutSamples.Reset(InNum);
+		OutSamples.AddUninitialized(InNum);
+		float* OutSampleData = OutSamples.GetData();
+
+		// Process audio one block at a time.
+		int32 LeftOver = InNum;
+		int32 BufferPos = 0;
+		while (LeftOver > 0)
+		{
+			int32 NumToProcess = FMath::Min(LeftOver, NumInternalBufferSamples);
+			ProcessAudioBlock(&InSampleData[BufferPos], NumToProcess, &OutSampleData[BufferPos]);
+			BufferPos += NumToProcess;
+			LeftOver -= NumToProcess;
+		}
 	}
-}
 
-void FIntegerDelay::ProcessAudioBlock(const float* InSamples, const int32 InNum, float* OutSamples)
-{
-	
-	// Update delay line.	
-	DelayLine->AddSamples(InSamples, InNum);
+	void FIntegerDelay::ProcessAudioBlock(const float* InSamples, const int32 InNum, float* OutSamples)
+	{
+		
+		// Update delay line.	
+		DelayLine->AddSamples(InSamples, InNum);
 
-	// Copy delayed version to output
-	const float* DelayData = DelayLine->InspectSamples(InNum + NumBufferOffsetSamples, NumDelayLineOffsetSamples);
+		// Copy delayed version to output
+		const float* DelayData = DelayLine->InspectSamples(InNum + NumBufferOffsetSamples, NumDelayLineOffsetSamples);
 
-	FMemory::Memcpy(OutSamples, &DelayData[NumBufferOffsetSamples], InNum * sizeof(float));
-	
-	// Remove unneeded delay line.
-	DelayLine->RemoveSamples(InNum);
-}
-
-void FIntegerDelay::PeekDelayLine(int32 InNum, Audio::FAlignedFloatBuffer& OutSamples)
-{
-	int32 NumToInspect = FMath::Min(DelayLine->GetNumAvailable(), NumDelaySamples);
-	NumToInspect = FMath::Min(InNum, NumToInspect);
+		FMemory::Memcpy(OutSamples, &DelayData[NumBufferOffsetSamples], InNum * sizeof(float));
+		
+		// Remove unneeded delay line.
+		DelayLine->RemoveSamples(InNum);
+	}
 
-	const float* DelaySamples = DelayLine->InspectSamples(NumToInspect);
-	if (nullptr == DelaySamples)
+	void FIntegerDelay::PeekDelayLine(int32 InNum, Audio::FAlignedFloatBuffer& OutSamples)
 	{
-		OutSamples.Reset(0);
-		return;
+		int32 NumToInspect = FMath::Min(DelayLine->GetNumAvailable(), NumDelaySamples);
+		NumToInspect = FMath::Min(InNum, NumToInspect);
+
+		const float* DelaySamples = DelayLine->InspectSamples(NumToInspect);
+		if (nullptr == DelaySamples)
+		{
+			OutSamples.Reset(0);
+			return;
+		}
+
+		OutSamples.Reset(NumToInspect);
+		OutSamples.AddUninitialized(NumToInspect);
+		FMemory::Memcpy(OutSamples.GetData(), DelaySamples, NumToInspect * sizeof(float));
 	}
-
-	OutSamples.Reset(NumToInspect);
-	OutSamples.AddUninitialized(NumToInspect);
-	FMemory::Memcpy(OutSamples.GetData(), DelaySamples, NumToInspect * sizeof(float));
 }
-
diff --git a/Engine/Source/Runtime/SignalProcessing/Private/LateReflectionsFast.cpp b/Engine/Source/Runtime/SignalProcessing/Private/LateReflectionsFast.cpp
--- a/Engine/Source/Runtime/SignalProcessing/Private/LateReflectionsFast.cpp
+++ b/Engine/Source/Runtime/SignalProcessing/Private/LateReflectionsFast.cpp
@@ -1,582 +1,584 @@
 // Copyright Epic Games, Inc. All Rights Reserved.
 
 #include "DSP/LateReflectionsFast.h"
-#include "CoreMinimal.h"
-#include "DSP/BufferVectorOperations.h"
 
-using namespace Audio;
+#include "DSP/BufferOnePoleLPF.h"
+#include "DSP/FloatArrayMath.h"
 
-// The presets for the late refelections reverb were based on a strange sample rate. This 
-// function helps translate delay values in the paper to delay values for this reverb.
-// https://ccrma.stanford.edu/~dattorro/EffectDesignPart1.pdf
-static int32 LateReflectionsGetDelaySamples(float InSampleRate, float InPresetValue) {
-	
-	return (int32)((float)InSampleRate / 29761.f * InPresetValue);
-}
-
-FLateReflectionsFastSettings::FLateReflectionsFastSettings()
-	: LateDelayMsec(0.0f)
-	, LateGainDB(0.0f)
-	, Bandwidth(0.5f)
-	, Diffusion(0.5f)
-	, Dampening(0.5f)
-	, Decay(0.5f)
-	, Density(0.5f)
-{}
-
-bool FLateReflectionsFastSettings::operator==(const FLateReflectionsFastSettings& Other) const
-{
-	bool bIsEqual = (
-		(Other.LateDelayMsec == LateDelayMsec) &&
-		(Other.LateGainDB == LateGainDB) &&
-		(Other.Bandwidth == Bandwidth) &&
-		(Other.Diffusion == Diffusion) &&
-		(Other.Dampening == Dampening) &&
-		(Other.Decay == Decay) &&
-		(Other.Density == Density));
-
-	return bIsEqual;
-}
-
-bool FLateReflectionsFastSettings::operator!=(const FLateReflectionsFastSettings& Other) const
-{
-	return !(*this == Other);
-}
-
-void FLateReflectionsPlateOutputs::ResizeAndZero(int32 InNumSamples)
+namespace Audio
 {
-	for (int32 i = 0; i < NumTaps; i++)
+	// The presets for the late refelections reverb were based on a strange sample rate. This 
+	// function helps translate delay values in the paper to delay values for this reverb.
+	// https://ccrma.stanford.edu/~dattorro/EffectDesignPart1.pdf
+	static int32 LateReflectionsGetDelaySamples(float InSampleRate, float InPresetValue) 
 	{
-		Taps[i].Reset(InNumSamples);
-		Taps[i].AddUninitialized(InNumSamples);
-		FMemory::Memset(Taps[i].GetData(), 0, sizeof(float) * InNumSamples);
+		return (int32)((float)InSampleRate / 29761.f * InPresetValue);
 	}
 
-	Output.Reset(InNumSamples);
-	Output.AddUninitialized(InNumSamples);
-	FMemory::Memset(Output.GetData(), 0, sizeof(float) * InNumSamples);
-}
-
-
-FLateReflectionsPlateDelays FLateReflectionsPlateDelays::DefaultLeftDelays(float InSampleRate)
-{
-	FLateReflectionsPlateDelays Delays;
-
-	Delays.NumSamplesModulatedBase = LateReflectionsGetDelaySamples(InSampleRate, 908);
-	Delays.NumSamplesModulatedDelta = LateReflectionsGetDelaySamples(InSampleRate, 16);
-	// Left Delay 1 should add up to 4217 and have tap outs at 353 (R0), 1190 (L4), and 3627 (R1)
-	Delays.NumSamplesDelayA = LateReflectionsGetDelaySamples(InSampleRate, 353);
-	Delays.NumSamplesDelayB = LateReflectionsGetDelaySamples(InSampleRate, 837);
-	Delays.NumSamplesDelayC = LateReflectionsGetDelaySamples(InSampleRate, 2437);
-	Delays.NumSamplesDelayD = LateReflectionsGetDelaySamples(InSampleRate, 590);
-	Delays.NumSamplesAPF = LateReflectionsGetDelaySamples(InSampleRate, 2656);
-	// APF Delay should have tapouts at 187 (L5) and 1228 (R2)
-	Delays.NumSamplesDelayE = LateReflectionsGetDelaySamples(InSampleRate, 187);
-	Delays.NumSamplesDelayF = LateReflectionsGetDelaySamples(InSampleRate, 1041);
-	// Left Delay 2 should add up to 3136 and have tap outs at 1066 (L6) and 2673 (R3)
-	Delays.NumSamplesDelayG = LateReflectionsGetDelaySamples(InSampleRate, 1066);
-	Delays.NumSamplesDelayH = LateReflectionsGetDelaySamples(InSampleRate, 1607);
-	Delays.NumSamplesDelayI = LateReflectionsGetDelaySamples(InSampleRate, 463);
-
-	return Delays;
-}
+	FLateReflectionsFastSettings::FLateReflectionsFastSettings()
+		: LateDelayMsec(0.0f)
+		, LateGainDB(0.0f)
+		, Bandwidth(0.5f)
+		, Diffusion(0.5f)
+		, Dampening(0.5f)
+		, Decay(0.5f)
+		, Density(0.5f)
+	{}
+
+	bool FLateReflectionsFastSettings::operator==(const FLateReflectionsFastSettings& Other) const
+	{
+		bool bIsEqual = (
+			(Other.LateDelayMsec == LateDelayMsec) &&
+			(Other.LateGainDB == LateGainDB) &&
+			(Other.Bandwidth == Bandwidth) &&
+			(Other.Diffusion == Diffusion) &&
+			(Other.Dampening == Dampening) &&
+			(Other.Decay == Decay) &&
+			(Other.Density == Density));
+
+		return bIsEqual;
+	}
 
-FLateReflectionsPlateDelays FLateReflectionsPlateDelays::DefaultRightDelays(float InSampleRate)
-{
-	FLateReflectionsPlateDelays Delays;
-
-	Delays.NumSamplesModulatedBase = LateReflectionsGetDelaySamples(InSampleRate, 672);
-	Delays.NumSamplesModulatedDelta = LateReflectionsGetDelaySamples(InSampleRate, 16);
-	// Right Delay 1 should add up to 4453 and have tap outs at 266 (L0), 2111 (R4), 2974 (L1)
-	Delays.NumSamplesDelayA = LateReflectionsGetDelaySamples(InSampleRate, 266);
-	Delays.NumSamplesDelayB = LateReflectionsGetDelaySamples(InSampleRate, 1845);
-	Delays.NumSamplesDelayC = LateReflectionsGetDelaySamples(InSampleRate, 863);
-	Delays.NumSamplesDelayD = LateReflectionsGetDelaySamples(InSampleRate, 1479);
-	Delays.NumSamplesAPF = LateReflectionsGetDelaySamples(InSampleRate, 1800);
-	// APF Delay should have tapouts at 335 (R5) and 1913 (L2)
-	Delays.NumSamplesDelayE = LateReflectionsGetDelaySamples(InSampleRate, 335);
-	Delays.NumSamplesDelayF = LateReflectionsGetDelaySamples(InSampleRate, 1578);
-	// Right Delay 2 should add up to 3720 and have tap outs at 121 (R6) and 1996 (L3)
-	Delays.NumSamplesDelayG = LateReflectionsGetDelaySamples(InSampleRate, 121);
-	Delays.NumSamplesDelayH = LateReflectionsGetDelaySamples(InSampleRate, 1875);
-	Delays.NumSamplesDelayI = LateReflectionsGetDelaySamples(InSampleRate, 1724);
-
-	return Delays;
-}
+	bool FLateReflectionsFastSettings::operator!=(const FLateReflectionsFastSettings& Other) const
+	{
+		return !(*this == Other);
+	}
 
+	void FLateReflectionsPlateOutputs::ResizeAndZero(int32 InNumSamples)
+	{
+		for (int32 i = 0; i < NumTaps; i++)
+		{
+			Taps[i].Reset(InNumSamples);
+			Taps[i].AddUninitialized(InNumSamples);
+			FMemory::Memset(Taps[i].GetData(), 0, sizeof(float) * InNumSamples);
+		}
 
-FLateReflectionsPlate::FLateReflectionsPlate(
-			float InSampleRate, 
-			int32 InMaxNumInternalBufferSamples, 
-			const FLateReflectionsPlateDelays& InDelays)
-: 	SampleRate(InSampleRate)
-	, NumInternalBufferSamples(InMaxNumInternalBufferSamples)
-	, Dampening(0.0005f)
-	, Decay(0.50f)
-	, Density(0.0f)
-	, PlateDelays(InDelays)
-{
-	// NumInternalBufferSamples must be less than last delay and should be aligned to SIMD boundaries
-	if (NumInternalBufferSamples > PlateDelays.NumSamplesDelayI) {
-		NumInternalBufferSamples = PlateDelays.NumSamplesDelayI;
+		Output.Reset(InNumSamples);
+		Output.AddUninitialized(InNumSamples);
+		FMemory::Memset(Output.GetData(), 0, sizeof(float) * InNumSamples);
 	}
-	NumInternalBufferSamples -= (NumInternalBufferSamples % AUDIO_NUM_FLOATS_PER_VECTOR_REGISTER);
-
-	checkf(NumInternalBufferSamples > 0, TEXT("InDelays.DelaySample2C too small"));
-
-	// Create the delay line for a plate
-	// ModulatedAPF -> Delay1[A-D] -> LPF -> APF -> Delay2[A-C] -> Output
-	//                                        |-> APFTapDelay[1-2]
-	ModulatedAPF = MakeUnique<FDynamicDelayAPF>(
-			-Density,
-			PlateDelays.NumSamplesModulatedBase - PlateDelays.NumSamplesModulatedDelta,
-			PlateDelays.NumSamplesModulatedBase + PlateDelays.NumSamplesModulatedDelta,
-			NumInternalBufferSamples,
-			SampleRate);
-	DelayA = MakeUnique<FIntegerDelay>(PlateDelays.NumSamplesDelayA, PlateDelays.NumSamplesDelayA);
-	DelayB = MakeUnique<FIntegerDelay>(PlateDelays.NumSamplesDelayB, PlateDelays.NumSamplesDelayB);
-	DelayC = MakeUnique<FIntegerDelay>(PlateDelays.NumSamplesDelayC, PlateDelays.NumSamplesDelayC);
-	DelayD = MakeUnique<FIntegerDelay>(PlateDelays.NumSamplesDelayD, PlateDelays.NumSamplesDelayD);
-	LPF = MakeUnique<FBufferOnePoleLPF>(Dampening);
-	APF = MakeUnique<FLongDelayAPF>(
-			Density - 0.15f,
-			PlateDelays.NumSamplesAPF,
-			NumInternalBufferSamples);
-	DelayE = MakeUnique<FIntegerDelay>(PlateDelays.NumSamplesDelayE, PlateDelays.NumSamplesDelayE);
-	DelayF = MakeUnique<FIntegerDelay>(PlateDelays.NumSamplesDelayF, PlateDelays.NumSamplesDelayF);
-	DelayG = MakeUnique<FIntegerDelay>(PlateDelays.NumSamplesDelayG, PlateDelays.NumSamplesDelayG);
-	DelayH = MakeUnique<FIntegerDelay>(PlateDelays.NumSamplesDelayH, PlateDelays.NumSamplesDelayH);
-	DelayI = MakeUnique<FIntegerDelay>(PlateDelays.NumSamplesDelayI, PlateDelays.NumSamplesDelayI);
-}
 
-FLateReflectionsPlate::~FLateReflectionsPlate()
-{
-}
 
+	FLateReflectionsPlateDelays FLateReflectionsPlateDelays::DefaultLeftDelays(float InSampleRate)
+	{
+		FLateReflectionsPlateDelays Delays;
+
+		Delays.NumSamplesModulatedBase = LateReflectionsGetDelaySamples(InSampleRate, 908);
+		Delays.NumSamplesModulatedDelta = LateReflectionsGetDelaySamples(InSampleRate, 16);
+		// Left Delay 1 should add up to 4217 and have tap outs at 353 (R0), 1190 (L4), and 3627 (R1)
+		Delays.NumSamplesDelayA = LateReflectionsGetDelaySamples(InSampleRate, 353);
+		Delays.NumSamplesDelayB = LateReflectionsGetDelaySamples(InSampleRate, 837);
+		Delays.NumSamplesDelayC = LateReflectionsGetDelaySamples(InSampleRate, 2437);
+		Delays.NumSamplesDelayD = LateReflectionsGetDelaySamples(InSampleRate, 590);
+		Delays.NumSamplesAPF = LateReflectionsGetDelaySamples(InSampleRate, 2656);
+		// APF Delay should have tapouts at 187 (L5) and 1228 (R2)
+		Delays.NumSamplesDelayE = LateReflectionsGetDelaySamples(InSampleRate, 187);
+		Delays.NumSamplesDelayF = LateReflectionsGetDelaySamples(InSampleRate, 1041);
+		// Left Delay 2 should add up to 3136 and have tap outs at 1066 (L6) and 2673 (R3)
+		Delays.NumSamplesDelayG = LateReflectionsGetDelaySamples(InSampleRate, 1066);
+		Delays.NumSamplesDelayH = LateReflectionsGetDelaySamples(InSampleRate, 1607);
+		Delays.NumSamplesDelayI = LateReflectionsGetDelaySamples(InSampleRate, 463);
+
+		return Delays;
+	}
 
-void FLateReflectionsPlate::ProcessAudioFrames(
-		const FAlignedFloatBuffer& InSamples,
-		const FAlignedFloatBuffer& InFeedbackSamples,
-		const FAlignedFloatBuffer& InDelayModulations,
-		FLateReflectionsPlateOutputs& OutPlateSamples)
-{
-	
-	// Prepare output
-	const int32 InNum = InSamples.Num();
-	OutPlateSamples.ResizeAndZero(InNum);
+	FLateReflectionsPlateDelays FLateReflectionsPlateDelays::DefaultRightDelays(float InSampleRate)
+	{
+		FLateReflectionsPlateDelays Delays;
+
+		Delays.NumSamplesModulatedBase = LateReflectionsGetDelaySamples(InSampleRate, 672);
+		Delays.NumSamplesModulatedDelta = LateReflectionsGetDelaySamples(InSampleRate, 16);
+		// Right Delay 1 should add up to 4453 and have tap outs at 266 (L0), 2111 (R4), 2974 (L1)
+		Delays.NumSamplesDelayA = LateReflectionsGetDelaySamples(InSampleRate, 266);
+		Delays.NumSamplesDelayB = LateReflectionsGetDelaySamples(InSampleRate, 1845);
+		Delays.NumSamplesDelayC = LateReflectionsGetDelaySamples(InSampleRate, 863);
+		Delays.NumSamplesDelayD = LateReflectionsGetDelaySamples(InSampleRate, 1479);
+		Delays.NumSamplesAPF = LateReflectionsGetDelaySamples(InSampleRate, 1800);
+		// APF Delay should have tapouts at 335 (R5) and 1913 (L2)
+		Delays.NumSamplesDelayE = LateReflectionsGetDelaySamples(InSampleRate, 335);
+		Delays.NumSamplesDelayF = LateReflectionsGetDelaySamples(InSampleRate, 1578);
+		// Right Delay 2 should add up to 3720 and have tap outs at 121 (R6) and 1996 (L3)
+		Delays.NumSamplesDelayG = LateReflectionsGetDelaySamples(InSampleRate, 121);
+		Delays.NumSamplesDelayH = LateReflectionsGetDelaySamples(InSampleRate, 1875);
+		Delays.NumSamplesDelayI = LateReflectionsGetDelaySamples(InSampleRate, 1724);
+
+		return Delays;
+	}
 
-	checkf(InNum == InFeedbackSamples.Num(), TEXT("InSamples and InFeedbackSamples must have equal length"));
-	checkf(InNum == InDelayModulations.Num(), TEXT("InSamples and InDelayModulations must have equal length"));
 
-	// Check buffer sizes and quit early if there is an error.
-	if (InNum != InFeedbackSamples.Num())
+	FLateReflectionsPlate::FLateReflectionsPlate(
+				float InSampleRate, 
+				int32 InMaxNumInternalBufferSamples, 
+				const FLateReflectionsPlateDelays& InDelays)
+	: 	SampleRate(InSampleRate)
+		, NumInternalBufferSamples(InMaxNumInternalBufferSamples)
+		, Dampening(0.0005f)
+		, Decay(0.50f)
+		, Density(0.0f)
+		, PlateDelays(InDelays)
 	{
-		return;	
+		// NumInternalBufferSamples must be less than last delay and should be aligned to SIMD boundaries
+		if (NumInternalBufferSamples > PlateDelays.NumSamplesDelayI) {
+			NumInternalBufferSamples = PlateDelays.NumSamplesDelayI;
+		}
+		NumInternalBufferSamples -= (NumInternalBufferSamples % AUDIO_NUM_FLOATS_PER_VECTOR_REGISTER);
+
+		checkf(NumInternalBufferSamples > 0, TEXT("InDelays.DelaySample2C too small"));
+
+		// Create the delay line for a plate
+		// ModulatedAPF -> Delay1[A-D] -> LPF -> APF -> Delay2[A-C] -> Output
+		//                                        |-> APFTapDelay[1-2]
+		ModulatedAPF = MakeUnique<FDynamicDelayAPF>(
+				-Density,
+				PlateDelays.NumSamplesModulatedBase - PlateDelays.NumSamplesModulatedDelta,
+				PlateDelays.NumSamplesModulatedBase + PlateDelays.NumSamplesModulatedDelta,
+				NumInternalBufferSamples,
+				SampleRate);
+		DelayA = MakeUnique<FIntegerDelay>(PlateDelays.NumSamplesDelayA, PlateDelays.NumSamplesDelayA);
+		DelayB = MakeUnique<FIntegerDelay>(PlateDelays.NumSamplesDelayB, PlateDelays.NumSamplesDelayB);
+		DelayC = MakeUnique<FIntegerDelay>(PlateDelays.NumSamplesDelayC, PlateDelays.NumSamplesDelayC);
+		DelayD = MakeUnique<FIntegerDelay>(PlateDelays.NumSamplesDelayD, PlateDelays.NumSamplesDelayD);
+		LPF = MakeUnique<FBufferOnePoleLPF>(Dampening);
+		APF = MakeUnique<FLongDelayAPF>(
+				Density - 0.15f,
+				PlateDelays.NumSamplesAPF,
+				NumInternalBufferSamples);
+		DelayE = MakeUnique<FIntegerDelay>(PlateDelays.NumSamplesDelayE, PlateDelays.NumSamplesDelayE);
+		DelayF = MakeUnique<FIntegerDelay>(PlateDelays.NumSamplesDelayF, PlateDelays.NumSamplesDelayF);
+		DelayG = MakeUnique<FIntegerDelay>(PlateDelays.NumSamplesDelayG, PlateDelays.NumSamplesDelayG);
+		DelayH = MakeUnique<FIntegerDelay>(PlateDelays.NumSamplesDelayH, PlateDelays.NumSamplesDelayH);
+		DelayI = MakeUnique<FIntegerDelay>(PlateDelays.NumSamplesDelayI, PlateDelays.NumSamplesDelayI);
 	}
-	if (InNum != InDelayModulations.Num())
+
+	FLateReflectionsPlate::~FLateReflectionsPlate()
 	{
-		return;
 	}
 
-	float* DataPtr = nullptr;
-
-	// Prepare internal buffers to handle input data.
-	WorkBufferA.Reset(InNum);
-	WorkBufferA.AddUninitialized(InNum);
-	WorkBufferB.Reset(InNum);
-	WorkBufferB.AddUninitialized(InNum);
-	WorkBufferC.Reset(InNum);
-	WorkBufferC.AddUninitialized(InNum);
-
-	// Copy Feedback Samples to aligned internal buffer.
-	ArrayWeightedSum(InFeedbackSamples, 1.0 - Decay, InSamples, WorkBufferA);
-
-	// Input -> ModulatedAPF
-	ModulatedAPF->ProcessAudio(WorkBufferA, InDelayModulations, WorkBufferB);
-
-	// ModulatedAPF -> Delay1
-	DelayA->ProcessAudio(WorkBufferB, OutPlateSamples.Taps[0]);
-	DelayB->ProcessAudio(OutPlateSamples.Taps[0], OutPlateSamples.Taps[1]);
-	DelayC->ProcessAudio(OutPlateSamples.Taps[1], OutPlateSamples.Taps[2]);
-	DelayD->ProcessAudio(OutPlateSamples.Taps[2], WorkBufferA);
-
-	// Apply dampening
-	ArrayMultiplyByConstantInPlace(WorkBufferA, 1.0f - Dampening);
-	
-	// Delay1 -> LPF
-	LPF->ProcessAudio(WorkBufferA, WorkBufferB);
-
-	// Apply decay
-	ArrayMultiplyByConstantInPlace(WorkBufferB, 1.0f - Decay);
-
-	// LPF -> APF
-	APF->ProcessAudio(WorkBufferB, WorkBufferA, WorkBufferC);
-	DelayE->ProcessAudio(WorkBufferC, OutPlateSamples.Taps[3]);
-	DelayF->ProcessAudio(OutPlateSamples.Taps[3], OutPlateSamples.Taps[4]);
-
-	// APF-> Delay2
-	DelayG->ProcessAudio(WorkBufferA, OutPlateSamples.Taps[5]);
-	DelayH->ProcessAudio(OutPlateSamples.Taps[5], OutPlateSamples.Taps[6]);
-	DelayI->ProcessAudio(OutPlateSamples.Taps[6], OutPlateSamples.Output);
-}
 
-void FLateReflectionsPlate::FlushAudio()
-{
-	DelayA->Reset();
-	DelayB->Reset();
-	DelayC->Reset();
-	DelayD->Reset();
-	DelayE->Reset();
-	DelayF->Reset();
-	DelayG->Reset();
-	DelayH->Reset();
-	DelayI->Reset();
-
-	LPF->FlushAudio();
-
-	ModulatedAPF->Reset();
-	APF->Reset();
-}
+	void FLateReflectionsPlate::ProcessAudioFrames(
+			const FAlignedFloatBuffer& InSamples,
+			const FAlignedFloatBuffer& InFeedbackSamples,
+			const FAlignedFloatBuffer& InDelayModulations,
+			FLateReflectionsPlateOutputs& OutPlateSamples)
+	{
+		
+		// Prepare output
+		const int32 InNum = InSamples.Num();
+		OutPlateSamples.ResizeAndZero(InNum);
 
-void FLateReflectionsPlate::SetDampening(float InDampening)
-{
-	Dampening = InDampening;
-	LPF->SetG(Dampening);
-}
+		checkf(InNum == InFeedbackSamples.Num(), TEXT("InSamples and InFeedbackSamples must have equal length"));
+		checkf(InNum == InDelayModulations.Num(), TEXT("InSamples and InDelayModulations must have equal length"));
 
-void FLateReflectionsPlate::SetDecay(float InDecay)
-{
-	Decay = InDecay;
-}
+		// Check buffer sizes and quit early if there is an error.
+		if (InNum != InFeedbackSamples.Num())
+		{
+			return;	
+		}
+		if (InNum != InDelayModulations.Num())
+		{
+			return;
+		}
 
-void FLateReflectionsPlate::SetDensity(float InDensity)
-{
-	Density = InDensity;
-	ModulatedAPF->SetG(FMath::Clamp(-Density, -0.9f, 0.9f));
-	APF->SetG(Density - 0.15f);
-}
+		float* DataPtr = nullptr;
 
-int32 FLateReflectionsPlate::GetNumInternalBufferSamples() const
-{
-	return NumInternalBufferSamples;
-}
+		// Prepare internal buffers to handle input data.
+		WorkBufferA.Reset(InNum);
+		WorkBufferA.AddUninitialized(InNum);
+		WorkBufferB.Reset(InNum);
+		WorkBufferB.AddUninitialized(InNum);
+		WorkBufferC.Reset(InNum);
+		WorkBufferC.AddUninitialized(InNum);
 
-void FLateReflectionsPlate::PeekDelayLine(int32 InNum, FAlignedFloatBuffer& OutSamples)
-{
-	DelayI->PeekDelayLine(InNum, OutSamples);
-}
+		// Copy Feedback Samples to aligned internal buffer.
+		ArrayWeightedSum(InFeedbackSamples, 1.0 - Decay, InSamples, WorkBufferA);
 
-// Limits on late reflections settings.
-const float FLateReflectionsFast::MaxLateDelayMsec = 2000.0f;
-const float FLateReflectionsFast::MinLateDelayMsec = 0.0f;
-const float FLateReflectionsFast::MaxLateGainDB = 0.0f;
-const float FLateReflectionsFast::MaxBandwidth = 0.99999f;
-const float FLateReflectionsFast::MinBandwidth = 0.0f;
-const float FLateReflectionsFast::MaxDampening = 0.99999f;
-const float FLateReflectionsFast::MinDampening = 0.0f;
-const float FLateReflectionsFast::MaxDiffusion = 0.99999f;
-const float FLateReflectionsFast::MinDiffusion = 0.0f;
-const float FLateReflectionsFast::MaxDecay = 1.0f;
-const float FLateReflectionsFast::MinDecay = 0.0001f;
-const float FLateReflectionsFast::MaxDensity = 1.0f;
-const float FLateReflectionsFast::MinDensity = 0.0f;
-
-const FLateReflectionsFastSettings FLateReflectionsFast::DefaultSettings;
-
-FLateReflectionsFast::FLateReflectionsFast(float InSampleRate, int32 InMaxNumInternalBufferSamples, const FLateReflectionsFastSettings& InSettings)
-:	SampleRate(InSampleRate)
-	, Gain(1.0f)
-	, ModulationPhase(0.0f)
-	, ModulationQuadPhase(PI / 2.0f)
-	, ModulationPhaseIncrement(0.0f)
-	, NumInternalBufferSamples(0)
-	, Settings(InSettings)
-	, LeftPlateDelays(FLateReflectionsPlateDelays::DefaultLeftDelays(InSampleRate))
-	, RightPlateDelays(FLateReflectionsPlateDelays::DefaultRightDelays(InSampleRate))
+		// Input -> ModulatedAPF
+		ModulatedAPF->ProcessAudio(WorkBufferA, InDelayModulations, WorkBufferB);
 
-{
-	// Copy and clamp settings.
-	Settings = InSettings;
-	ClampSettings(Settings);
-	
-	// Create plates. 
-	LeftPlate = MakeUnique<FLateReflectionsPlate>(SampleRate, InMaxNumInternalBufferSamples, LeftPlateDelays);
-	RightPlate = MakeUnique<FLateReflectionsPlate>(SampleRate, InMaxNumInternalBufferSamples, RightPlateDelays);
-
-	// The block size is limited by the delay lengths in the left & right plates.
-	// Internal block size will be set to minimum of allowable sizes from plates.
-	NumInternalBufferSamples = FMath::Min(LeftPlate->GetNumInternalBufferSamples(), RightPlate->GetNumInternalBufferSamples());
-
-	// Set maximum predelay. Add 8 samples to give a little extra buffer for floating point rounding
-	// differences. 
-	int32 MaxDelay = (int32)(SampleRate * FLateReflectionsFast::MaxLateDelayMsec / 1000.0f) + 8;
-	PreDelay = MakeUnique<FIntegerDelay>(MaxDelay, 0, NumInternalBufferSamples);
-
-	InputLPF = MakeUnique<FBufferOnePoleLPF>(1.0f - Settings.Bandwidth);
-
-	// make the signal decorrelators
-	APF1 = MakeUnique<FLongDelayAPF>(Settings.Diffusion, LateReflectionsGetDelaySamples(SampleRate, 142), NumInternalBufferSamples);
-	APF2 = MakeUnique<FLongDelayAPF>(Settings.Diffusion, LateReflectionsGetDelaySamples(SampleRate, 107), NumInternalBufferSamples);
-	APF3 = MakeUnique<FLongDelayAPF>(Settings.Diffusion - 0.125f, LateReflectionsGetDelaySamples(SampleRate, 379), NumInternalBufferSamples);
-	APF4 = MakeUnique<FLongDelayAPF>(Settings.Diffusion - 0.125f, LateReflectionsGetDelaySamples(SampleRate, 277), NumInternalBufferSamples);
-
-	// Set modulation rate to 1 Hz. 
-	ModulationPhaseIncrement = 2.0 * PI / SampleRate;
-	
-	// Resize internal buffers
-	WorkBufferA.Reset(NumInternalBufferSamples);
-	WorkBufferB.Reset(NumInternalBufferSamples);
-	WorkBufferC.Reset(NumInternalBufferSamples);
-	LeftDelayModSamples.Reset(NumInternalBufferSamples);
-	RightDelayModSamples.Reset(NumInternalBufferSamples);
-
-	WorkBufferA.AddUninitialized(NumInternalBufferSamples);
-	WorkBufferB.AddUninitialized(NumInternalBufferSamples);
-	WorkBufferC.AddUninitialized(NumInternalBufferSamples);
-	LeftDelayModSamples.AddUninitialized(NumInternalBufferSamples);
-	RightDelayModSamples.AddUninitialized(NumInternalBufferSamples);
-
-	LeftPlateOutputs.ResizeAndZero(NumInternalBufferSamples);
-	RightPlateOutputs.ResizeAndZero(NumInternalBufferSamples);
-
-
-	ApplySettings();
-}
+		// ModulatedAPF -> Delay1
+		DelayA->ProcessAudio(WorkBufferB, OutPlateSamples.Taps[0]);
+		DelayB->ProcessAudio(OutPlateSamples.Taps[0], OutPlateSamples.Taps[1]);
+		DelayC->ProcessAudio(OutPlateSamples.Taps[1], OutPlateSamples.Taps[2]);
+		DelayD->ProcessAudio(OutPlateSamples.Taps[2], WorkBufferA);
 
-FLateReflectionsFast::~FLateReflectionsFast()
-{}
+		// Apply dampening
+		ArrayMultiplyByConstantInPlace(WorkBufferA, 1.0f - Dampening);
+		
+		// Delay1 -> LPF
+		LPF->ProcessAudio(WorkBufferA, WorkBufferB);
 
-void FLateReflectionsFast::ClampSettings(FLateReflectionsFastSettings& InOutSettings)
-{
-	// Enforce settings to be within max/min
-	InOutSettings.LateDelayMsec = FMath::Clamp(InOutSettings.LateDelayMsec, MinLateDelayMsec, MaxLateDelayMsec);
-	InOutSettings.LateGainDB = FMath::Min(InOutSettings.LateGainDB, FLateReflectionsFast::MaxLateGainDB);
-	InOutSettings.Bandwidth = FMath::Clamp(InOutSettings.Bandwidth, MinBandwidth, MaxBandwidth);
-	InOutSettings.Dampening = FMath::Clamp(InOutSettings.Dampening, MinDampening, MaxDampening);
-	InOutSettings.Diffusion = FMath::Clamp(InOutSettings.Diffusion, MinDiffusion, MaxDiffusion);
-	InOutSettings.Decay = FMath::Clamp(InOutSettings.Decay, MinDecay, MaxDecay);
-	InOutSettings.Density = FMath::Clamp(InOutSettings.Density, MinDensity, MaxDensity);
-}
+		// Apply decay
+		ArrayMultiplyByConstantInPlace(WorkBufferB, 1.0f - Decay);
 
-// Sets the reverb settings, applies, and updates
-void FLateReflectionsFast::SetSettings(const FLateReflectionsFastSettings& InSettings)
-{
-	Settings = InSettings;
-	ClampSettings(Settings);
-	ApplySettings();
-}
+		// LPF -> APF
+		APF->ProcessAudio(WorkBufferB, WorkBufferA, WorkBufferC);
+		DelayE->ProcessAudio(WorkBufferC, OutPlateSamples.Taps[3]);
+		DelayF->ProcessAudio(OutPlateSamples.Taps[3], OutPlateSamples.Taps[4]);
 
-void FLateReflectionsFast::ProcessAudio(const FAlignedFloatBuffer& InSamples, const int32 InNumChannels, FAlignedFloatBuffer& OutLeftSamples, FAlignedFloatBuffer& OutRightSamples)
-{
-	checkf((InNumChannels == 1) || (InNumChannels == 2), TEXT("FLateReflections only supports 1 or 2 channel input audio"));
+		// APF-> Delay2
+		DelayG->ProcessAudio(WorkBufferA, OutPlateSamples.Taps[5]);
+		DelayH->ProcessAudio(OutPlateSamples.Taps[5], OutPlateSamples.Taps[6]);
+		DelayI->ProcessAudio(OutPlateSamples.Taps[6], OutPlateSamples.Output);
+	}
 
-	const int32 InNum = InSamples.Num();
-	const int32 InNumFrames = InNum / InNumChannels;
-	
-	OutLeftSamples.Reset(InNumFrames);
-	OutRightSamples.Reset(InNumFrames);
+	void FLateReflectionsPlate::FlushAudio()
+	{
+		DelayA->Reset();
+		DelayB->Reset();
+		DelayC->Reset();
+		DelayD->Reset();
+		DelayE->Reset();
+		DelayF->Reset();
+		DelayG->Reset();
+		DelayH->Reset();
+		DelayI->Reset();
+
+		LPF->FlushAudio();
+
+		ModulatedAPF->Reset();
+		APF->Reset();
+	}
 
-	OutLeftSamples.AddUninitialized(InNumFrames);
-	OutRightSamples.AddUninitialized(InNumFrames);
+	void FLateReflectionsPlate::SetDampening(float InDampening)
+	{
+		Dampening = InDampening;
+		LPF->SetG(Dampening);
+	}
 
-	const float* InSampleData = InSamples.GetData();
-	float* OutLeftSampleData = OutLeftSamples.GetData();
-	float* OutRightSampleData = OutRightSamples.GetData();
+	void FLateReflectionsPlate::SetDecay(float InDecay)
+	{
+		Decay = InDecay;
+	}
 
-	int32 FramesLeftOver = InNumFrames;
-	int32 InPos = 0;
-	int32 OutPos = 0;
-	while (FramesLeftOver > 0)
+	void FLateReflectionsPlate::SetDensity(float InDensity)
 	{
-		int32 FramesToProcess = FMath::Min(NumInternalBufferSamples, FramesLeftOver);
+		Density = InDensity;
+		ModulatedAPF->SetG(FMath::Clamp(-Density, -0.9f, 0.9f));
+		APF->SetG(Density - 0.15f);
+	}
 
-		ProcessAudioBuffer(&InSampleData[InPos], FramesToProcess, InNumChannels, &OutLeftSampleData[OutPos], &OutRightSampleData[OutPos]);
+	int32 FLateReflectionsPlate::GetNumInternalBufferSamples() const
+	{
+		return NumInternalBufferSamples;
+	}
 
-		FramesLeftOver -= FramesToProcess;
-		InPos += FramesToProcess * InNumChannels;
-		OutPos += FramesToProcess;
+	void FLateReflectionsPlate::PeekDelayLine(int32 InNum, FAlignedFloatBuffer& OutSamples)
+	{
+		DelayI->PeekDelayLine(InNum, OutSamples);
 	}
-}
 
-void FLateReflectionsFast::FlushAudio()
-{
-	PreDelay->Reset();
+	// Limits on late reflections settings.
+	const float FLateReflectionsFast::MaxLateDelayMsec = 2000.0f;
+	const float FLateReflectionsFast::MinLateDelayMsec = 0.0f;
+	const float FLateReflectionsFast::MaxLateGainDB = 0.0f;
+	const float FLateReflectionsFast::MaxBandwidth = 0.99999f;
+	const float FLateReflectionsFast::MinBandwidth = 0.0f;
+	const float FLateReflectionsFast::MaxDampening = 0.99999f;
+	const float FLateReflectionsFast::MinDampening = 0.0f;
+	const float FLateReflectionsFast::MaxDiffusion = 0.99999f;
+	const float FLateReflectionsFast::MinDiffusion = 0.0f;
+	const float FLateReflectionsFast::MaxDecay = 1.0f;
+	const float FLateReflectionsFast::MinDecay = 0.0001f;
+	const float FLateReflectionsFast::MaxDensity = 1.0f;
+	const float FLateReflectionsFast::MinDensity = 0.0f;
+
+	const FLateReflectionsFastSettings FLateReflectionsFast::DefaultSettings;
+
+	FLateReflectionsFast::FLateReflectionsFast(float InSampleRate, int32 InMaxNumInternalBufferSamples, const FLateReflectionsFastSettings& InSettings)
+	:	SampleRate(InSampleRate)
+		, Gain(1.0f)
+		, ModulationPhase(0.0f)
+		, ModulationQuadPhase(PI / 2.0f)
+		, ModulationPhaseIncrement(0.0f)
+		, NumInternalBufferSamples(0)
+		, Settings(InSettings)
+		, LeftPlateDelays(FLateReflectionsPlateDelays::DefaultLeftDelays(InSampleRate))
+		, RightPlateDelays(FLateReflectionsPlateDelays::DefaultRightDelays(InSampleRate))
 
-	InputLPF->FlushAudio();
+	{
+		// Copy and clamp settings.
+		Settings = InSettings;
+		ClampSettings(Settings);
+		
+		// Create plates. 
+		LeftPlate = MakeUnique<FLateReflectionsPlate>(SampleRate, InMaxNumInternalBufferSamples, LeftPlateDelays);
+		RightPlate = MakeUnique<FLateReflectionsPlate>(SampleRate, InMaxNumInternalBufferSamples, RightPlateDelays);
 
-	APF1->Reset();
-	APF2->Reset();
-	APF3->Reset();
-	APF4->Reset();
+		// The block size is limited by the delay lengths in the left & right plates.
+		// Internal block size will be set to minimum of allowable sizes from plates.
+		NumInternalBufferSamples = FMath::Min(LeftPlate->GetNumInternalBufferSamples(), RightPlate->GetNumInternalBufferSamples());
 
-	LeftPlate->FlushAudio();
-	RightPlate->FlushAudio();
-}
+		// Set maximum predelay. Add 8 samples to give a little extra buffer for floating point rounding
+		// differences. 
+		int32 MaxDelay = (int32)(SampleRate * FLateReflectionsFast::MaxLateDelayMsec / 1000.0f) + 8;
+		PreDelay = MakeUnique<FIntegerDelay>(MaxDelay, 0, NumInternalBufferSamples);
 
+		InputLPF = MakeUnique<FBufferOnePoleLPF>(1.0f - Settings.Bandwidth);
 
-void FLateReflectionsFast::ProcessAudioBuffer(const float* InSampleData, const int32 InNumFrames, const int32 InNumChannels, float* OutLeftSampleData, float* OutRightSampleData)
-{
-	checkf((1 == InNumChannels) || (2 == InNumChannels), TEXT("FLateReflections only supports 1 and 2 channel input audio"));
+		// make the signal decorrelators
+		APF1 = MakeUnique<FLongDelayAPF>(Settings.Diffusion, LateReflectionsGetDelaySamples(SampleRate, 142), NumInternalBufferSamples);
+		APF2 = MakeUnique<FLongDelayAPF>(Settings.Diffusion, LateReflectionsGetDelaySamples(SampleRate, 107), NumInternalBufferSamples);
+		APF3 = MakeUnique<FLongDelayAPF>(Settings.Diffusion - 0.125f, LateReflectionsGetDelaySamples(SampleRate, 379), NumInternalBufferSamples);
+		APF4 = MakeUnique<FLongDelayAPF>(Settings.Diffusion - 0.125f, LateReflectionsGetDelaySamples(SampleRate, 277), NumInternalBufferSamples);
+
+		// Set modulation rate to 1 Hz. 
+		ModulationPhaseIncrement = 2.0 * PI / SampleRate;
+		
+		// Resize internal buffers
+		WorkBufferA.Reset(NumInternalBufferSamples);
+		WorkBufferB.Reset(NumInternalBufferSamples);
+		WorkBufferC.Reset(NumInternalBufferSamples);
+		LeftDelayModSamples.Reset(NumInternalBufferSamples);
+		RightDelayModSamples.Reset(NumInternalBufferSamples);
+
+		WorkBufferA.AddUninitialized(NumInternalBufferSamples);
+		WorkBufferB.AddUninitialized(NumInternalBufferSamples);
+		WorkBufferC.AddUninitialized(NumInternalBufferSamples);
+		LeftDelayModSamples.AddUninitialized(NumInternalBufferSamples);
+		RightDelayModSamples.AddUninitialized(NumInternalBufferSamples);
 
-	// Resize internal buffers
-	WorkBufferA.Reset(InNumFrames);
-	WorkBufferB.Reset(InNumFrames);
-	LeftDelayModSamples.Reset(InNumFrames);
-	RightDelayModSamples.Reset(InNumFrames);
+		LeftPlateOutputs.ResizeAndZero(NumInternalBufferSamples);
+		RightPlateOutputs.ResizeAndZero(NumInternalBufferSamples);
 
-	WorkBufferA.AddUninitialized(InNumFrames);
-	WorkBufferB.AddUninitialized(InNumFrames);
-	LeftDelayModSamples.AddUninitialized(InNumFrames);
-	RightDelayModSamples.AddUninitialized(InNumFrames);
 
-	// Downmix to mono
-	if (InNumChannels == 1)
+		ApplySettings();
+	}
+
+	FLateReflectionsFast::~FLateReflectionsFast()
+	{}
+
+	void FLateReflectionsFast::ClampSettings(FLateReflectionsFastSettings& InOutSettings)
 	{
-		FMemory::Memcpy(WorkBufferA.GetData(), InSampleData, sizeof(float) * InNumFrames);
+		// Enforce settings to be within max/min
+		InOutSettings.LateDelayMsec = FMath::Clamp(InOutSettings.LateDelayMsec, MinLateDelayMsec, MaxLateDelayMsec);
+		InOutSettings.LateGainDB = FMath::Min(InOutSettings.LateGainDB, FLateReflectionsFast::MaxLateGainDB);
+		InOutSettings.Bandwidth = FMath::Clamp(InOutSettings.Bandwidth, MinBandwidth, MaxBandwidth);
+		InOutSettings.Dampening = FMath::Clamp(InOutSettings.Dampening, MinDampening, MaxDampening);
+		InOutSettings.Diffusion = FMath::Clamp(InOutSettings.Diffusion, MinDiffusion, MaxDiffusion);
+		InOutSettings.Decay = FMath::Clamp(InOutSettings.Decay, MinDecay, MaxDecay);
+		InOutSettings.Density = FMath::Clamp(InOutSettings.Density, MinDensity, MaxDensity);
 	}
-	else if (InNumChannels == 2)
+
+	// Sets the reverb settings, applies, and updates
+	void FLateReflectionsFast::SetSettings(const FLateReflectionsFastSettings& InSettings)
 	{
-		BufferSum2ChannelToMonoFast(InSampleData, WorkBufferA.GetData(), InNumFrames);
+		Settings = InSettings;
+		ClampSettings(Settings);
+		ApplySettings();
 	}
-	else
+
+	void FLateReflectionsFast::ProcessAudio(const FAlignedFloatBuffer& InSamples, const int32 InNumChannels, FAlignedFloatBuffer& OutLeftSamples, FAlignedFloatBuffer& OutRightSamples)
 	{
-		// Shouldn't reach this line if checks are respected. 
-		FMemory::Memset(OutLeftSampleData, 0, sizeof(float) * InNumFrames);
-		FMemory::Memset(OutRightSampleData, 0, sizeof(float) * InNumFrames);
-		return;
+		checkf((InNumChannels == 1) || (InNumChannels == 2), TEXT("FLateReflections only supports 1 or 2 channel input audio"));
+
+		const int32 InNum = InSamples.Num();
+		const int32 InNumFrames = InNum / InNumChannels;
+		
+		OutLeftSamples.Reset(InNumFrames);
+		OutRightSamples.Reset(InNumFrames);
+
+		OutLeftSamples.AddUninitialized(InNumFrames);
+		OutRightSamples.AddUninitialized(InNumFrames);
+
+		const float* InSampleData = InSamples.GetData();
+		float* OutLeftSampleData = OutLeftSamples.GetData();
+		float* OutRightSampleData = OutRightSamples.GetData();
+
+		int32 FramesLeftOver = InNumFrames;
+		int32 InPos = 0;
+		int32 OutPos = 0;
+		while (FramesLeftOver > 0)
+		{
+			int32 FramesToProcess = FMath::Min(NumInternalBufferSamples, FramesLeftOver);
+
+			ProcessAudioBuffer(&InSampleData[InPos], FramesToProcess, InNumChannels, &OutLeftSampleData[OutPos], &OutRightSampleData[OutPos]);
+
+			FramesLeftOver -= FramesToProcess;
+			InPos += FramesToProcess * InNumChannels;
+			OutPos += FramesToProcess;
+		}
 	}
 
-	// Apply bandwidth, gain and channel averaging multipliers
-	ArrayMultiplyByConstantInPlace(WorkBufferA, Settings.Bandwidth * Gain / InNumChannels);
-	
-	// Predelay
-	PreDelay->ProcessAudio(WorkBufferA, WorkBufferB);
-
-	// Input Diffusion 
-	InputLPF->ProcessAudio(WorkBufferB, WorkBufferA);
-	APF1->ProcessAudio(WorkBufferA, WorkBufferB);
-	APF2->ProcessAudio(WorkBufferB, WorkBufferA);
-	APF3->ProcessAudio(WorkBufferA, WorkBufferB);
-	APF4->ProcessAudio(WorkBufferB, WorkBufferA);
-	
-	// The plates are connected to eachother so that the output of the left plate is fed into the input of the
-	// right plate and vice versa. The delay lines of the plates are stored before processing new data through
-	// the plate.
-	LeftPlate->PeekDelayLine(InNumFrames, WorkBufferB);
-	RightPlate->PeekDelayLine(InNumFrames, WorkBufferC);
-	
-	// Create the delay line modulations.
-	GeneraterPlateModulations(InNumFrames, LeftDelayModSamples, RightDelayModSamples);
-
-	// Run audio through plates
-	LeftPlate->ProcessAudioFrames(WorkBufferA, WorkBufferC, LeftDelayModSamples, LeftPlateOutputs);
-	RightPlate->ProcessAudioFrames(WorkBufferA, WorkBufferB, RightDelayModSamples, RightPlateOutputs);
-
-	TArrayView<float> OutLeftSampleDataView(OutLeftSampleData, InNumFrames);
-	TArrayView<float> OutRightSampleDataView(OutRightSampleData, InNumFrames);
-
-	TArrayView<const float> LeftTaps1View(LeftPlateOutputs.Taps[1].GetData(), InNumFrames);
-	TArrayView<const float> LeftTaps2View(LeftPlateOutputs.Taps[2].GetData(), InNumFrames);
-	TArrayView<const float> LeftTaps3View(LeftPlateOutputs.Taps[3].GetData(), InNumFrames);
-	TArrayView<const float> LeftTaps4View(LeftPlateOutputs.Taps[4].GetData(), InNumFrames);
-	TArrayView<const float> LeftTaps5View(LeftPlateOutputs.Taps[5].GetData(), InNumFrames);
-	TArrayView<const float> LeftTaps6View(LeftPlateOutputs.Taps[6].GetData(), InNumFrames);
-
-	TArrayView<const float> RightTaps1View(RightPlateOutputs.Taps[1].GetData(), InNumFrames);
-	TArrayView<const float> RightTaps2View(RightPlateOutputs.Taps[2].GetData(), InNumFrames);
-	TArrayView<const float> RightTaps3View(RightPlateOutputs.Taps[3].GetData(), InNumFrames);
-	TArrayView<const float> RightTaps4View(RightPlateOutputs.Taps[4].GetData(), InNumFrames);
-	TArrayView<const float> RightTaps5View(RightPlateOutputs.Taps[5].GetData(), InNumFrames);
-	TArrayView<const float> RightTaps6View(RightPlateOutputs.Taps[6].GetData(), InNumFrames);
-
-	// Left Output
-	// Channel -> [Plate and Tap#]
-	// L -> +R0
-	// L -> +R2
-	// L -> -R4
-	// L -> +R6
-	// L -> -L1
-	// L -> -L3
-	// L -> -L5
-	FMemory::Memcpy(OutLeftSampleData, RightPlateOutputs.Taps[0].GetData(), InNumFrames * sizeof(float));
-	ArrayMixIn(RightTaps2View, OutLeftSampleDataView);
-	ArraySubtractInPlace2(OutLeftSampleDataView, RightTaps4View);
-	ArrayMixIn(RightTaps6View, OutLeftSampleDataView);
-	ArraySubtractInPlace2(OutLeftSampleDataView, LeftTaps1View);
-	ArraySubtractInPlace2(OutLeftSampleDataView, LeftTaps3View);
-	ArraySubtractInPlace2(OutLeftSampleDataView, LeftTaps5View);
-	
-	// Right Output
-	// Channel -> [Plate and Tap#]
-	// R -> +L0
-	// R -> +L2
-	// R -> -L4
-	// R -> +L6
-	// R -> -R1
-	// R -> -R3
-	// R -> -R5
-	FMemory::Memcpy(OutRightSampleData, LeftPlateOutputs.Taps[0].GetData(), InNumFrames * sizeof(float));
-	ArrayMixIn(LeftTaps2View, OutRightSampleDataView);
-	ArraySubtractInPlace2(OutRightSampleDataView, LeftTaps4View);
-	ArrayMixIn(LeftTaps6View, OutRightSampleDataView);
-	ArraySubtractInPlace2(OutRightSampleDataView, RightTaps1View);
-	ArraySubtractInPlace2(OutRightSampleDataView, LeftTaps3View);
-	ArraySubtractInPlace2(OutRightSampleDataView, LeftTaps5View);
-}
+	void FLateReflectionsFast::FlushAudio()
+	{
+		PreDelay->Reset();
 
-void FLateReflectionsFast::ApplySettings()
-{
-	// Convert gain to linear. 
-	Gain = FMath::Pow(10.0f, Settings.LateGainDB/ 20.0f);
+		InputLPF->FlushAudio();
 
-	// Convert delay to samples
-	PreDelay->SetDelayLengthSamples((int32)(Settings.LateDelayMsec * SampleRate / 1000.0f));
+		APF1->Reset();
+		APF2->Reset();
+		APF3->Reset();
+		APF4->Reset();
 
-	InputLPF->SetG(1.0f - Settings.Bandwidth);
+		LeftPlate->FlushAudio();
+		RightPlate->FlushAudio();
+	}
 
-	APF1->SetG(Settings.Diffusion);
-	APF2->SetG(Settings.Diffusion);
-	APF3->SetG(Settings.Diffusion - 0.125f);
-	APF4->SetG(Settings.Diffusion - 0.125f);
 
-	LeftPlate->SetDensity(Settings.Density);
-	LeftPlate->SetDampening(Settings.Dampening);
-	LeftPlate->SetDecay(Settings.Decay);
+	void FLateReflectionsFast::ProcessAudioBuffer(const float* InSampleData, const int32 InNumFrames, const int32 InNumChannels, float* OutLeftSampleData, float* OutRightSampleData)
+	{
+		checkf((1 == InNumChannels) || (2 == InNumChannels), TEXT("FLateReflections only supports 1 and 2 channel input audio"));
 
-	RightPlate->SetDensity(Settings.Density);
-	RightPlate->SetDampening(Settings.Dampening);
-	RightPlate->SetDecay(Settings.Decay);
-}
+		// Resize internal buffers
+		WorkBufferA.Reset(InNumFrames);
+		WorkBufferB.Reset(InNumFrames);
+		LeftDelayModSamples.Reset(InNumFrames);
+		RightDelayModSamples.Reset(InNumFrames);
 
+		WorkBufferA.AddUninitialized(InNumFrames);
+		WorkBufferB.AddUninitialized(InNumFrames);
+		LeftDelayModSamples.AddUninitialized(InNumFrames);
+		RightDelayModSamples.AddUninitialized(InNumFrames);
 
-void FLateReflectionsFast::GeneraterPlateModulations(const int32 InNum, FAlignedFloatBuffer& OutLeftDelays, FAlignedFloatBuffer& OutRightDelays)
-{
-	OutLeftDelays.Reset(InNum);
-	OutRightDelays.Reset(InNum);
-	OutLeftDelays.AddUninitialized(InNum);
-	OutRightDelays.AddUninitialized(InNum);
-
-	// Generate quadrature plate modulations to help break up hi-freq modes in plates. 
-	float* LeftModData = OutLeftDelays.GetData();
-	float* RightModData = OutRightDelays.GetData();
-	for (int32 i = 0; i < InNum; i++)
-	{		
-		float NormalPhaseOut = 0.5f * FastSin(ModulationPhase) + 0.5f;
-		float QuadPhaseOut = 0.5f * FastSin(ModulationQuadPhase) + 0.5f;
-		ModulationPhase += ModulationPhaseIncrement;
-
-		if (ModulationPhase > PI)
+		// Downmix to mono
+		if (InNumChannels == 1)
 		{
-			ModulationPhase -= 2 * PI;
+			FMemory::Memcpy(WorkBufferA.GetData(), InSampleData, sizeof(float) * InNumFrames);
 		}
-
-		ModulationQuadPhase += ModulationPhaseIncrement;
-		if (ModulationQuadPhase > PI)
+		else if (InNumChannels == 2)
 		{
-			ModulationQuadPhase -= 2 * PI;
+			BufferSum2ChannelToMonoFast(InSampleData, WorkBufferA.GetData(), InNumFrames);
 		}
+		else
+		{
+			// Shouldn't reach this line if checks are respected. 
+			FMemory::Memset(OutLeftSampleData, 0, sizeof(float) * InNumFrames);
+			FMemory::Memset(OutRightSampleData, 0, sizeof(float) * InNumFrames);
+			return;
+		}
+
+		// Apply bandwidth, gain and channel averaging multipliers
+		ArrayMultiplyByConstantInPlace(WorkBufferA, Settings.Bandwidth * Gain / InNumChannels);
+		
+		// Predelay
+		PreDelay->ProcessAudio(WorkBufferA, WorkBufferB);
+
+		// Input Diffusion 
+		InputLPF->ProcessAudio(WorkBufferB, WorkBufferA);
+		APF1->ProcessAudio(WorkBufferA, WorkBufferB);
+		APF2->ProcessAudio(WorkBufferB, WorkBufferA);
+		APF3->ProcessAudio(WorkBufferA, WorkBufferB);
+		APF4->ProcessAudio(WorkBufferB, WorkBufferA);
 		
-		LeftModData[i] = LeftPlateDelays.NumSamplesModulatedBase + (NormalPhaseOut * LeftPlateDelays.NumSamplesModulatedDelta);
-		RightModData[i] = RightPlateDelays.NumSamplesModulatedBase + (QuadPhaseOut * RightPlateDelays.NumSamplesModulatedDelta);
+		// The plates are connected to eachother so that the output of the left plate is fed into the input of the
+		// right plate and vice versa. The delay lines of the plates are stored before processing new data through
+		// the plate.
+		LeftPlate->PeekDelayLine(InNumFrames, WorkBufferB);
+		RightPlate->PeekDelayLine(InNumFrames, WorkBufferC);
+		
+		// Create the delay line modulations.
+		GeneraterPlateModulations(InNumFrames, LeftDelayModSamples, RightDelayModSamples);
+
+		// Run audio through plates
+		LeftPlate->ProcessAudioFrames(WorkBufferA, WorkBufferC, LeftDelayModSamples, LeftPlateOutputs);
+		RightPlate->ProcessAudioFrames(WorkBufferA, WorkBufferB, RightDelayModSamples, RightPlateOutputs);
+
+		TArrayView<float> OutLeftSampleDataView(OutLeftSampleData, InNumFrames);
+		TArrayView<float> OutRightSampleDataView(OutRightSampleData, InNumFrames);
+
+		TArrayView<const float> LeftTaps1View(LeftPlateOutputs.Taps[1].GetData(), InNumFrames);
+		TArrayView<const float> LeftTaps2View(LeftPlateOutputs.Taps[2].GetData(), InNumFrames);
+		TArrayView<const float> LeftTaps3View(LeftPlateOutputs.Taps[3].GetData(), InNumFrames);
+		TArrayView<const float> LeftTaps4View(LeftPlateOutputs.Taps[4].GetData(), InNumFrames);
+		TArrayView<const float> LeftTaps5View(LeftPlateOutputs.Taps[5].GetData(), InNumFrames);
+		TArrayView<const float> LeftTaps6View(LeftPlateOutputs.Taps[6].GetData(), InNumFrames);
+
+		TArrayView<const float> RightTaps1View(RightPlateOutputs.Taps[1].GetData(), InNumFrames);
+		TArrayView<const float> RightTaps2View(RightPlateOutputs.Taps[2].GetData(), InNumFrames);
+		TArrayView<const float> RightTaps3View(RightPlateOutputs.Taps[3].GetData(), InNumFrames);
+		TArrayView<const float> RightTaps4View(RightPlateOutputs.Taps[4].GetData(), InNumFrames);
+		TArrayView<const float> RightTaps5View(RightPlateOutputs.Taps[5].GetData(), InNumFrames);
+		TArrayView<const float> RightTaps6View(RightPlateOutputs.Taps[6].GetData(), InNumFrames);
+
+		// Left Output
+		// Channel -> [Plate and Tap#]
+		// L -> +R0
+		// L -> +R2
+		// L -> -R4
+		// L -> +R6
+		// L -> -L1
+		// L -> -L3
+		// L -> -L5
+		FMemory::Memcpy(OutLeftSampleData, RightPlateOutputs.Taps[0].GetData(), InNumFrames * sizeof(float));
+		ArrayMixIn(RightTaps2View, OutLeftSampleDataView);
+		ArraySubtractInPlace2(OutLeftSampleDataView, RightTaps4View);
+		ArrayMixIn(RightTaps6View, OutLeftSampleDataView);
+		ArraySubtractInPlace2(OutLeftSampleDataView, LeftTaps1View);
+		ArraySubtractInPlace2(OutLeftSampleDataView, LeftTaps3View);
+		ArraySubtractInPlace2(OutLeftSampleDataView, LeftTaps5View);
+		
+		// Right Output
+		// Channel -> [Plate and Tap#]
+		// R -> +L0
+		// R -> +L2
+		// R -> -L4
+		// R -> +L6
+		// R -> -R1
+		// R -> -R3
+		// R -> -R5
+		FMemory::Memcpy(OutRightSampleData, LeftPlateOutputs.Taps[0].GetData(), InNumFrames * sizeof(float));
+		ArrayMixIn(LeftTaps2View, OutRightSampleDataView);
+		ArraySubtractInPlace2(OutRightSampleDataView, LeftTaps4View);
+		ArrayMixIn(LeftTaps6View, OutRightSampleDataView);
+		ArraySubtractInPlace2(OutRightSampleDataView, RightTaps1View);
+		ArraySubtractInPlace2(OutRightSampleDataView, LeftTaps3View);
+		ArraySubtractInPlace2(OutRightSampleDataView, LeftTaps5View);
+	}
+
+	void FLateReflectionsFast::ApplySettings()
+	{
+		// Convert gain to linear. 
+		Gain = FMath::Pow(10.0f, Settings.LateGainDB/ 20.0f);
+
+		// Convert delay to samples
+		PreDelay->SetDelayLengthSamples((int32)(Settings.LateDelayMsec * SampleRate / 1000.0f));
+
+		InputLPF->SetG(1.0f - Settings.Bandwidth);
+
+		APF1->SetG(Settings.Diffusion);
+		APF2->SetG(Settings.Diffusion);
+		APF3->SetG(Settings.Diffusion - 0.125f);
+		APF4->SetG(Settings.Diffusion - 0.125f);
+
+		LeftPlate->SetDensity(Settings.Density);
+		LeftPlate->SetDampening(Settings.Dampening);
+		LeftPlate->SetDecay(Settings.Decay);
+
+		RightPlate->SetDensity(Settings.Density);
+		RightPlate->SetDampening(Settings.Dampening);
+		RightPlate->SetDecay(Settings.Decay);
+	}
+
+
+	void FLateReflectionsFast::GeneraterPlateModulations(const int32 InNum, FAlignedFloatBuffer& OutLeftDelays, FAlignedFloatBuffer& OutRightDelays)
+	{
+		OutLeftDelays.Reset(InNum);
+		OutRightDelays.Reset(InNum);
+		OutLeftDelays.AddUninitialized(InNum);
+		OutRightDelays.AddUninitialized(InNum);
+
+		// Generate quadrature plate modulations to help break up hi-freq modes in plates. 
+		float* LeftModData = OutLeftDelays.GetData();
+		float* RightModData = OutRightDelays.GetData();
+		for (int32 i = 0; i < InNum; i++)
+		{		
+			float NormalPhaseOut = 0.5f * FastSin(ModulationPhase) + 0.5f;
+			float QuadPhaseOut = 0.5f * FastSin(ModulationQuadPhase) + 0.5f;
+			ModulationPhase += ModulationPhaseIncrement;
+
+			if (ModulationPhase > PI)
+			{
+				ModulationPhase -= 2 * PI;
+			}
+
+			ModulationQuadPhase += ModulationPhaseIncrement;
+			if (ModulationQuadPhase > PI)
+			{
+				ModulationQuadPhase -= 2 * PI;
+			}
+			
+			LeftModData[i] = LeftPlateDelays.NumSamplesModulatedBase + (NormalPhaseOut * LeftPlateDelays.NumSamplesModulatedDelta);
+			RightModData[i] = RightPlateDelays.NumSamplesModulatedBase + (QuadPhaseOut * RightPlateDelays.NumSamplesModulatedDelta);
+		}
 	}
 }
diff --git a/Engine/Source/Runtime/SignalProcessing/Private/LinearInterpFractionalDelay.cpp b/Engine/Source/Runtime/SignalProcessing/Private/LinearInterpFractionalDelay.cpp
--- a/Engine/Source/Runtime/SignalProcessing/Private/LinearInterpFractionalDelay.cpp
+++ b/Engine/Source/Runtime/SignalProcessing/Private/LinearInterpFractionalDelay.cpp
@@ -4,160 +4,156 @@
 #include "DSP/Dsp.h"
 #include "DSP/BufferVectorOperations.h"
 
-using namespace Audio;
-
-FLinearInterpFractionalDelay::FLinearInterpFractionalDelay(int32 InMaxDelay, int32 InMaxNumInternalBufferSamples)
-: MaxDelay(InMaxDelay)
-, NumInternalBufferSamples(InMaxNumInternalBufferSamples)
-, UpperDelayPos(nullptr)
-, LowerDelayPos(nullptr)
+namespace Audio
 {
-	checkf(MaxDelay > 0, TEXT("InMaxDelay must be greater than zero"));
-	if (MaxDelay < 1)
+	FLinearInterpFractionalDelay::FLinearInterpFractionalDelay(int32 InMaxDelay, int32 InMaxNumInternalBufferSamples)
+	: MaxDelay(InMaxDelay)
+	, NumInternalBufferSamples(InMaxNumInternalBufferSamples)
+	, UpperDelayPos(nullptr)
+	, LowerDelayPos(nullptr)
 	{
-		MaxDelay = 1;
-	}
+		checkf(MaxDelay > 0, TEXT("InMaxDelay must be greater than zero"));
+		if (MaxDelay < 1)
+		{
+			MaxDelay = 1;
+		}
 
-	while (0 != (NumInternalBufferSamples % AUDIO_NUM_FLOATS_PER_VECTOR_REGISTER))
-	{
-		NumInternalBufferSamples--;
-	}
-	if (NumInternalBufferSamples < 1)
-	{
-		NumInternalBufferSamples = AUDIO_NUM_FLOATS_PER_VECTOR_REGISTER;
-	}
+		while (0 != (NumInternalBufferSamples % AUDIO_NUM_FLOATS_PER_VECTOR_REGISTER))
+		{
+			NumInternalBufferSamples--;
+		}
+		if (NumInternalBufferSamples < 1)
+		{
+			NumInternalBufferSamples = AUDIO_NUM_FLOATS_PER_VECTOR_REGISTER;
+		}
 
-	// Allocate and prepare delay line for maximum delay.
-	DelayLine = MakeUnique<FAlignedBlockBuffer>((2 * (MaxDelay + 1)) + NumInternalBufferSamples, MaxDelay + NumInternalBufferSamples + 1);
-	DelayLine->AddZeros(MaxDelay + 1);
-	
-	IntegerDelayOffsets.Reset(NumInternalBufferSamples);
-	IntegerDelayOffsets.AddUninitialized(NumInternalBufferSamples);
+		// Allocate and prepare delay line for maximum delay.
+		DelayLine = MakeUnique<FAlignedBlockBuffer>((2 * (MaxDelay + 1)) + NumInternalBufferSamples, MaxDelay + NumInternalBufferSamples + 1);
+		DelayLine->AddZeros(MaxDelay + 1);
+		
+		IntegerDelayOffsets.Reset(NumInternalBufferSamples);
+		IntegerDelayOffsets.AddUninitialized(NumInternalBufferSamples);
 
-	for (int32 i = 0; i < NumInternalBufferSamples; i++)
-	{
-		IntegerDelayOffsets[i] = i + MaxDelay;
-	}
+		for (int32 i = 0; i < NumInternalBufferSamples; i++)
+		{
+			IntegerDelayOffsets[i] = i + MaxDelay;
+		}
 
-	UpperDelayPos = (int*)FMemory::Malloc(4 * sizeof(int), AUDIO_BUFFER_ALIGNMENT);
-	LowerDelayPos = (int*)FMemory::Malloc(4 * sizeof(int), AUDIO_BUFFER_ALIGNMENT);
-}
+		UpperDelayPos = (int*)FMemory::Malloc(4 * sizeof(int), AUDIO_BUFFER_ALIGNMENT);
+		LowerDelayPos = (int*)FMemory::Malloc(4 * sizeof(int), AUDIO_BUFFER_ALIGNMENT);
+	}
 
 
-// Destructor
-FLinearInterpFractionalDelay::~FLinearInterpFractionalDelay()
-{
-	FMemory::Free(UpperDelayPos);
-	FMemory::Free(LowerDelayPos);
-}
+	// Destructor
+	FLinearInterpFractionalDelay::~FLinearInterpFractionalDelay()
+	{
+		FMemory::Free(UpperDelayPos);
+		FMemory::Free(LowerDelayPos);
+	}
 
 
-// Resets the delay line state, flushes buffer and resets read/write pointers.
-void FLinearInterpFractionalDelay::Reset() 
-{
-	DelayLine->ClearSamples();
-	DelayLine->AddZeros(MaxDelay + 1);
-}
+	// Resets the delay line state, flushes buffer and resets read/write pointers.
+	void FLinearInterpFractionalDelay::Reset() 
+	{
+		DelayLine->ClearSamples();
+		DelayLine->AddZeros(MaxDelay + 1);
+	}
 
 
-void FLinearInterpFractionalDelay::ProcessAudio(const FAlignedFloatBuffer& InSamples, const FAlignedFloatBuffer& InDelays, FAlignedFloatBuffer& OutSamples)
-{
-	const int32 InNum = InSamples.Num();
-	checkf(InNum == InDelays.Num(), TEXT("Input buffers must be equal length"));
+	void FLinearInterpFractionalDelay::ProcessAudio(const FAlignedFloatBuffer& InSamples, const FAlignedFloatBuffer& InDelays, FAlignedFloatBuffer& OutSamples)
+	{
+		const int32 InNum = InSamples.Num();
+		checkf(InNum == InDelays.Num(), TEXT("Input buffers must be equal length"));
 
-	// Prepare output buffer
-	OutSamples.Reset(InNum);
-	OutSamples.AddUninitialized(InNum);
+		// Prepare output buffer
+		OutSamples.Reset(InNum);
+		OutSamples.AddUninitialized(InNum);
 
-	if (InNum != InDelays.Num())
-	{
-		// Return empty buffer on invalid input.
-		if (InNum > 0)
+		if (InNum != InDelays.Num())
 		{
-			FMemory::Memset(OutSamples.GetData(), 0, sizeof(float) * InNum);
+			// Return empty buffer on invalid input.
+			if (InNum > 0)
+			{
+				FMemory::Memset(OutSamples.GetData(), 0, sizeof(float) * InNum);
+			}
+			return;
 		}
-		return;
-	}
 
 
-	float* OutSampleData = OutSamples.GetData();
-	const float* InSampleData = InSamples.GetData();
-	const float* InDelayData = InDelays.GetData();
+		float* OutSampleData = OutSamples.GetData();
+		const float* InSampleData = InSamples.GetData();
+		const float* InDelayData = InDelays.GetData();
 
-	// Process audio one block at a time.
-	int32 LeftOver = InNum;
-	int32 BufferPos = 0;
-	while (LeftOver > 0)
-	{
-		int32 NumToProcess = FMath::Min(LeftOver, NumInternalBufferSamples);
-		ProcessAudioBlock(&InSampleData[BufferPos], &InDelayData[BufferPos], NumToProcess, &OutSampleData[BufferPos]);
-		BufferPos += NumToProcess;
-		LeftOver -= NumToProcess;
+		// Process audio one block at a time.
+		int32 LeftOver = InNum;
+		int32 BufferPos = 0;
+		while (LeftOver > 0)
+		{
+			int32 NumToProcess = FMath::Min(LeftOver, NumInternalBufferSamples);
+			ProcessAudioBlock(&InSampleData[BufferPos], &InDelayData[BufferPos], NumToProcess, &OutSampleData[BufferPos]);
+			BufferPos += NumToProcess;
+			LeftOver -= NumToProcess;
+		}
 	}
-}
 
-void FLinearInterpFractionalDelay::ProcessAudioBlock(const float* InSamples, const float* InDelays, const int32 InNum, float* OutSamples)
-{
-	checkf(0 == (InNum % 4), TEXT("Array length must be multiple of 4"));
-	checkf(IsAligned<const float*>(InSamples, AUDIO_SIMD_BYTE_ALIGNMENT), TEXT("Memory must be aligned to use vector operations."));
-	checkf(IsAligned<const float*>(InDelays, AUDIO_SIMD_BYTE_ALIGNMENT), TEXT("Memory must be aligned to use vector operations."));
-	checkf(IsAligned<float*>(OutSamples, AUDIO_SIMD_BYTE_ALIGNMENT), TEXT("Memory must be aligned to use vector operations."));
-
-
-	// Update delay line.
-	DelayLine->AddSamples(InSamples, InNum);
-
-	const float* DelayData = DelayLine->InspectSamples(InNum + MaxDelay + 1);
-	const int32* IntegerDelayOffsetData = IntegerDelayOffsets.GetData();
-
-	const VectorRegister4Float VMaxDelay = MakeVectorRegister((float)MaxDelay, (float)MaxDelay, (float)MaxDelay, (float)MaxDelay);
-	for (int32 i = 0; i < InNum; i += 4)
+	void FLinearInterpFractionalDelay::ProcessAudioBlock(const float* InSamples, const float* InDelays, const int32 InNum, float* OutSamples)
 	{
-		
-		VectorRegister4Float VFractionalDelays = VectorLoadAligned(&InDelays[i]);
-		// Ensure fractional delays are positive
-		VFractionalDelays = VectorMax(VFractionalDelays, GlobalVectorConstants::FloatZero);
-		VFractionalDelays = VectorMin(VFractionalDelays, VMaxDelay);
-
-		// Separate integer from fraction
-		VectorRegister4Float VFloorDelays = VectorFloor(VFractionalDelays);
+		checkf(0 == (InNum % 4), TEXT("Array length must be multiple of 4"));
 
-		// Determine linear weights
-		VectorRegister4Float VUpperCoefficients = VectorSubtract(VFractionalDelays, VFloorDelays);
-		VectorRegister4Float VLowerCoefficients = VectorSubtract(GlobalVectorConstants::FloatOne, VUpperCoefficients);
+		// Update delay line.
+		DelayLine->AddSamples(InSamples, InNum);
 
+		const float* DelayData = DelayLine->InspectSamples(InNum + MaxDelay + 1);
+		const int32* IntegerDelayOffsetData = IntegerDelayOffsets.GetData();
 
-		// Make integer locations relative to block
-		VectorRegister4Int VIntegerDelays = VectorFloatToInt(VFloorDelays);
-		VectorRegister4Int VIntegerDelayOffset = VectorIntLoadAligned(&IntegerDelayOffsetData[i]);
-		VIntegerDelays = VectorIntSubtract(VIntegerDelayOffset, VIntegerDelays);
+		const VectorRegister4Float VMaxDelay = MakeVectorRegister((float)MaxDelay, (float)MaxDelay, (float)MaxDelay, (float)MaxDelay);
+		for (int32 i = 0; i < InNum; i += 4)
+		{
+			
+			VectorRegister4Float VFractionalDelays = VectorLoad(&InDelays[i]);
+			// Ensure fractional delays are positive
+			VFractionalDelays = VectorMax(VFractionalDelays, GlobalVectorConstants::FloatZero);
+			VFractionalDelays = VectorMin(VFractionalDelays, VMaxDelay);
+
+			// Separate integer from fraction
+			VectorRegister4Float VFloorDelays = VectorFloor(VFractionalDelays);
+
+			// Determine linear weights
+			VectorRegister4Float VUpperCoefficients = VectorSubtract(VFractionalDelays, VFloorDelays);
+			VectorRegister4Float VLowerCoefficients = VectorSubtract(GlobalVectorConstants::FloatOne, VUpperCoefficients);
+
+
+			// Make integer locations relative to block
+			VectorRegister4Int VIntegerDelays = VectorFloatToInt(VFloorDelays);
+			VectorRegister4Int VIntegerDelayOffset = VectorIntLoadAligned(&IntegerDelayOffsetData[i]);
+			VIntegerDelays = VectorIntSubtract(VIntegerDelayOffset, VIntegerDelays);
+
+			// Lookup samples for interpolation
+			VectorIntStoreAligned(VIntegerDelays, UpperDelayPos);
+			VectorIntStoreAligned(VectorIntAdd(VIntegerDelays, GlobalVectorConstants::IntOne), LowerDelayPos);
+			
+			VectorRegister4Float VLowerSamples = MakeVectorRegister(
+				DelayData[LowerDelayPos[0]],
+				DelayData[LowerDelayPos[1]],
+				DelayData[LowerDelayPos[2]],
+				DelayData[LowerDelayPos[3]]
+			);
+			VectorRegister4Float VUpperSamples = MakeVectorRegister(
+				DelayData[UpperDelayPos[0]],
+				DelayData[UpperDelayPos[1]],
+				DelayData[UpperDelayPos[2]],
+				DelayData[UpperDelayPos[3]]
+			);
+
+			// Interpolate samples
+			VectorRegister4Float VOut = VectorMultiplyAdd(
+				VLowerSamples,
+				VLowerCoefficients,
+				VectorMultiply(VUpperSamples, VUpperCoefficients));
+			VectorStore(VOut, &OutSamples[i]);
+		}
 
-		// Lookup samples for interpolation
-		VectorIntStoreAligned(VIntegerDelays, UpperDelayPos);
-		VectorIntStoreAligned(VectorIntAdd(VIntegerDelays, GlobalVectorConstants::IntOne), LowerDelayPos);
-		
-		VectorRegister4Float VLowerSamples = MakeVectorRegister(
-			DelayData[LowerDelayPos[0]],
-			DelayData[LowerDelayPos[1]],
-			DelayData[LowerDelayPos[2]],
-			DelayData[LowerDelayPos[3]]
-		);
-		VectorRegister4Float VUpperSamples = MakeVectorRegister(
-			DelayData[UpperDelayPos[0]],
-			DelayData[UpperDelayPos[1]],
-			DelayData[UpperDelayPos[2]],
-			DelayData[UpperDelayPos[3]]
-		);
-
-		// Interpolate samples
-		VectorRegister4Float VOut = VectorMultiplyAdd(
-			VLowerSamples,
-			VLowerCoefficients,
-			VectorMultiply(VUpperSamples, VUpperCoefficients));
-		VectorStoreAligned(VOut, &OutSamples[i]);
+		// Remove unneeded delay line.
+		DelayLine->RemoveSamples(InNum);
 	}
-
-	// Remove unneeded delay line.
-	DelayLine->RemoveSamples(InNum);
 }
-
diff --git a/Engine/Source/Runtime/SignalProcessing/Private/LinkwitzRileyBandSplitter.cpp b/Engine/Source/Runtime/SignalProcessing/Private/LinkwitzRileyBandSplitter.cpp
--- a/Engine/Source/Runtime/SignalProcessing/Private/LinkwitzRileyBandSplitter.cpp
+++ b/Engine/Source/Runtime/SignalProcessing/Private/LinkwitzRileyBandSplitter.cpp
@@ -1,7 +1,7 @@
 // Copyright Epic Games, Inc. All Rights Reserved.
 
 #include "DSP/LinkwitzRileyBandSplitter.h"
-#include "DSP/BufferVectorOperations.h"
+#include "DSP/FloatArrayMath.h"
 
 #ifndef TWO_PI
 #define TWO_PI (6.28318530718)
diff --git a/Engine/Source/Runtime/SignalProcessing/Private/LongDelayAPF.cpp b/Engine/Source/Runtime/SignalProcessing/Private/LongDelayAPF.cpp
--- a/Engine/Source/Runtime/SignalProcessing/Private/LongDelayAPF.cpp
+++ b/Engine/Source/Runtime/SignalProcessing/Private/LongDelayAPF.cpp
@@ -4,171 +4,171 @@
 #include "DSP/BufferVectorOperations.h"
 #include "DSP/Dsp.h"
 
-using namespace Audio;
-
-FLongDelayAPF::FLongDelayAPF(float InG, int32 InNumDelaySamples, int32 InMaxNumInternalBufferSamples)
-:	G(InG),
-	NumDelaySamples(InNumDelaySamples),
-	NumInternalBufferSamples(InMaxNumInternalBufferSamples)
+namespace Audio
 {
-	// NumInternalBufferSamples must be less than or equal to the length of the delay
-	if (NumInternalBufferSamples > NumDelaySamples)
+	FLongDelayAPF::FLongDelayAPF(float InG, int32 InNumDelaySamples, int32 InMaxNumInternalBufferSamples)
+	:	G(InG),
+		NumDelaySamples(InNumDelaySamples),
+		NumInternalBufferSamples(InMaxNumInternalBufferSamples)
 	{
-		NumInternalBufferSamples = NumDelaySamples;
-
-		// Block size must be divisible by simd alignment to support simd operations.
-		NumInternalBufferSamples -= (NumInternalBufferSamples % AUDIO_NUM_FLOATS_PER_VECTOR_REGISTER);
-	}
-
-	checkf(NumInternalBufferSamples > 1, TEXT("Invalid internal buffer length"));
+		// NumInternalBufferSamples must be less than or equal to the length of the delay
+		if (NumInternalBufferSamples > NumDelaySamples)
+		{
+			NumInternalBufferSamples = NumDelaySamples;
 
-	// Allocate delay line
-	DelayLine = MakeUnique<FAlignedBlockBuffer>(2 * NumInternalBufferSamples + NumDelaySamples, NumInternalBufferSamples);
-	DelayLine->AddZeros(NumDelaySamples);
+			// Block size must be divisible by simd alignment to support simd operations.
+			NumInternalBufferSamples -= (NumInternalBufferSamples % AUDIO_NUM_FLOATS_PER_VECTOR_REGISTER);
+		}
 
-	// Allocate internal buffer
-	WorkBuffer.Reset(NumInternalBufferSamples);
-	WorkBuffer.AddUninitialized(NumInternalBufferSamples);
-}
+		checkf(NumInternalBufferSamples > 1, TEXT("Invalid internal buffer length"));
 
-FLongDelayAPF::~FLongDelayAPF()
-{}
+		// Allocate delay line
+		DelayLine = MakeUnique<FAlignedBlockBuffer>(2 * NumInternalBufferSamples + NumDelaySamples, NumInternalBufferSamples);
+		DelayLine->AddZeros(NumDelaySamples);
 
-void FLongDelayAPF::ProcessAudio(FAlignedFloatBuffer& Samples)
-{
-	const int32 InNum = Samples.Num();
+		// Allocate internal buffer
+		WorkBuffer.Reset(NumInternalBufferSamples);
+		WorkBuffer.AddUninitialized(NumInternalBufferSamples);
+	}
 
-	float* InData = Samples.GetData();
-	float* OutDelayData = WorkBuffer.GetData();
+	FLongDelayAPF::~FLongDelayAPF()
+	{}
 
-	// Process audio one block at a time.
-	int32 LeftOver = InNum;
-	int32 BufferIndex = 0;
-	while (LeftOver != 0)
+	void FLongDelayAPF::ProcessAudio(FAlignedFloatBuffer& Samples)
 	{
-		// Determine block size for this loop.
-		int32 NumToProcess = FMath::Min<int32>(NumInternalBufferSamples, LeftOver);
-		const float* InDelayData = DelayLine->InspectSamples(NumToProcess);
-
-		ProcessAudioBlock(&InData[BufferIndex], InDelayData, NumToProcess, &InData[BufferIndex], OutDelayData);
+		const int32 InNum = Samples.Num();
 
-		// Update delay line with new data.
-		DelayLine->RemoveSamples(NumToProcess);
-		DelayLine->AddSamples(OutDelayData, NumToProcess);
+		float* InData = Samples.GetData();
+		float* OutDelayData = WorkBuffer.GetData();
 
-		LeftOver -= NumToProcess;
-		BufferIndex += NumToProcess;
-	}
-}
+		// Process audio one block at a time.
+		int32 LeftOver = InNum;
+		int32 BufferIndex = 0;
+		while (LeftOver != 0)
+		{
+			// Determine block size for this loop.
+			int32 NumToProcess = FMath::Min<int32>(NumInternalBufferSamples, LeftOver);
+			const float* InDelayData = DelayLine->InspectSamples(NumToProcess);
 
-void FLongDelayAPF::ProcessAudio(const FAlignedFloatBuffer& InSamples, FAlignedFloatBuffer& OutSamples)
-{
-	const float* InData = InSamples.GetData();
-	const int32 InNum = InSamples.Num();
-	float* OutDelayData = WorkBuffer.GetData();
+			ProcessAudioBlock(&InData[BufferIndex], InDelayData, NumToProcess, &InData[BufferIndex], OutDelayData);
 
-	// Resize output
-	OutSamples.Reset(InNum);
-	OutSamples.AddUninitialized(InNum);
-	float* OutData = OutSamples.GetData();
+			// Update delay line with new data.
+			DelayLine->RemoveSamples(NumToProcess);
+			DelayLine->AddSamples(OutDelayData, NumToProcess);
 
+			LeftOver -= NumToProcess;
+			BufferIndex += NumToProcess;
+		}
+	}
 
-	// Process audio one block at a time.
-	int32 LeftOver = InNum;
-	int32 BufferIndex = 0;
-	while (LeftOver != 0) 
+	void FLongDelayAPF::ProcessAudio(const FAlignedFloatBuffer& InSamples, FAlignedFloatBuffer& OutSamples)
 	{
-		// Determine block size for this loop.
-		int32 NumToProcess = FMath::Min<int32>(NumInternalBufferSamples, LeftOver);
-		const float* InDelayData = DelayLine->InspectSamples(NumToProcess);
-
-		ProcessAudioBlock(&InData[BufferIndex], InDelayData, NumToProcess, &OutData[BufferIndex], OutDelayData);
-
-		// Update delay line with new data.
-		DelayLine->RemoveSamples(NumToProcess);
-		DelayLine->AddSamples(OutDelayData, NumToProcess);
-
-		LeftOver -= NumToProcess;
-		BufferIndex += NumToProcess;
+		const float* InData = InSamples.GetData();
+		const int32 InNum = InSamples.Num();
+		float* OutDelayData = WorkBuffer.GetData();
+
+		// Resize output
+		OutSamples.Reset(InNum);
+		OutSamples.AddUninitialized(InNum);
+		float* OutData = OutSamples.GetData();
+
+
+		// Process audio one block at a time.
+		int32 LeftOver = InNum;
+		int32 BufferIndex = 0;
+		while (LeftOver != 0) 
+		{
+			// Determine block size for this loop.
+			int32 NumToProcess = FMath::Min<int32>(NumInternalBufferSamples, LeftOver);
+			const float* InDelayData = DelayLine->InspectSamples(NumToProcess);
+
+			ProcessAudioBlock(&InData[BufferIndex], InDelayData, NumToProcess, &OutData[BufferIndex], OutDelayData);
+
+			// Update delay line with new data.
+			DelayLine->RemoveSamples(NumToProcess);
+			DelayLine->AddSamples(OutDelayData, NumToProcess);
+
+			LeftOver -= NumToProcess;
+			BufferIndex += NumToProcess;
+		}
 	}
-}
 
-void FLongDelayAPF::ProcessAudio(const FAlignedFloatBuffer& InSamples, FAlignedFloatBuffer& OutSamples, FAlignedFloatBuffer& OutDelaySamples)
-{
-	const float* InData = InSamples.GetData();
-	const int32 InNum = InSamples.Num();
-
-	// Resize output buffers
-	OutSamples.Reset(InNum);
-	OutSamples.AddUninitialized(InNum);
-	OutDelaySamples.Reset(InNum);
-	OutDelaySamples.AddUninitialized(InNum);
-
-	float* OutData = OutSamples.GetData();
-	float* OutDelayData = OutDelaySamples.GetData();
-
-	// Process audio one block at a time.
-	int32 LeftOver = InNum;
-	int32 BufferIndex = 0;
-	while (LeftOver != 0) 
+	void FLongDelayAPF::ProcessAudio(const FAlignedFloatBuffer& InSamples, FAlignedFloatBuffer& OutSamples, FAlignedFloatBuffer& OutDelaySamples)
 	{
-		int32 NumToProcess = FMath::Min<int32>(NumInternalBufferSamples, LeftOver);
-		const float* InDelayData = DelayLine->InspectSamples(NumToProcess);
-
-		ProcessAudioBlock(&InData[BufferIndex], InDelayData, NumToProcess, &OutData[BufferIndex], &OutDelayData[BufferIndex]);
-		
-		// Update internal delay line.
-		DelayLine->RemoveSamples(NumToProcess);
-		DelayLine->AddSamples(&OutDelayData[BufferIndex], NumToProcess);
-
-		LeftOver -= NumToProcess;
-		BufferIndex += NumToProcess;
+		const float* InData = InSamples.GetData();
+		const int32 InNum = InSamples.Num();
+
+		// Resize output buffers
+		OutSamples.Reset(InNum);
+		OutSamples.AddUninitialized(InNum);
+		OutDelaySamples.Reset(InNum);
+		OutDelaySamples.AddUninitialized(InNum);
+
+		float* OutData = OutSamples.GetData();
+		float* OutDelayData = OutDelaySamples.GetData();
+
+		// Process audio one block at a time.
+		int32 LeftOver = InNum;
+		int32 BufferIndex = 0;
+		while (LeftOver != 0) 
+		{
+			int32 NumToProcess = FMath::Min<int32>(NumInternalBufferSamples, LeftOver);
+			const float* InDelayData = DelayLine->InspectSamples(NumToProcess);
+
+			ProcessAudioBlock(&InData[BufferIndex], InDelayData, NumToProcess, &OutData[BufferIndex], &OutDelayData[BufferIndex]);
+			
+			// Update internal delay line.
+			DelayLine->RemoveSamples(NumToProcess);
+			DelayLine->AddSamples(&OutDelayData[BufferIndex], NumToProcess);
+
+			LeftOver -= NumToProcess;
+			BufferIndex += NumToProcess;
+		}
 	}
-}
 
-void FLongDelayAPF::ProcessAudioBlock(const float* InSamples, const float* InDelaySamples, const int32 InNum, float* OutSamples, float* OutDelaySamples)
-{
-	// Calculate new delay line samples. "w[n] = x[n] + gw[n - d]"
-	int32 NumToSIMD = InNum - (InNum % AUDIO_NUM_FLOATS_PER_VECTOR_REGISTER);
-
-	VectorRegister4Float VG = MakeVectorRegisterFloat(G, G, G, G);
-	VectorRegister4Float VNG = MakeVectorRegisterFloat(-G, -G, -G, -G);
-	VectorRegister4Float VFMIN = MakeVectorRegisterFloat(FLT_MIN, FLT_MIN, FLT_MIN, FLT_MIN);
-	VectorRegister4Float VNFMIN = MakeVectorRegisterFloat(-FLT_MIN, -FLT_MIN, -FLT_MIN, -FLT_MIN);
-
-	for (int32 i = 0; i < InNum; i += 4)
+	void FLongDelayAPF::ProcessAudioBlock(const float* InSamples, const float* InDelaySamples, const int32 InNum, float* OutSamples, float* OutDelaySamples)
 	{
-		VectorRegister4Float VInDelay = VectorLoadAligned(&InDelaySamples[i]);
-		VectorRegister4Float VInSamples = VectorLoadAligned(&InSamples[i]);
-		// w[n] = x[n] + G * w[n - D]
-		VectorRegister4Float VOutDelay = VectorMultiplyAdd(VInDelay, VG, VInSamples);
-		
-		// Underflow clamp
-		VectorRegister4Float Mask = VectorBitwiseAnd(VectorCompareGT(VOutDelay, VNFMIN), VectorCompareLT(VOutDelay, VFMIN));
-		VOutDelay = VectorSelect(Mask, GlobalVectorConstants::FloatZero, VOutDelay);
-		VectorStoreAligned(VOutDelay, &OutDelaySamples[i]);
-
-		// y[n] = -G * w[n] + w[n - D]
-		VectorRegister4Float VOut = VectorMultiplyAdd(VOutDelay, VNG, VInDelay);
-		VectorStoreAligned(VOut, &OutSamples[i]);
+		// Calculate new delay line samples. "w[n] = x[n] + gw[n - d]"
+		int32 NumToSIMD = InNum - (InNum % AUDIO_NUM_FLOATS_PER_VECTOR_REGISTER);
+
+		VectorRegister4Float VG = MakeVectorRegisterFloat(G, G, G, G);
+		VectorRegister4Float VNG = MakeVectorRegisterFloat(-G, -G, -G, -G);
+		VectorRegister4Float VFMIN = MakeVectorRegisterFloat(FLT_MIN, FLT_MIN, FLT_MIN, FLT_MIN);
+		VectorRegister4Float VNFMIN = MakeVectorRegisterFloat(-FLT_MIN, -FLT_MIN, -FLT_MIN, -FLT_MIN);
+
+		for (int32 i = 0; i < InNum; i += 4)
+		{
+			VectorRegister4Float VInDelay = VectorLoadAligned(&InDelaySamples[i]);
+			VectorRegister4Float VInSamples = VectorLoadAligned(&InSamples[i]);
+			// w[n] = x[n] + G * w[n - D]
+			VectorRegister4Float VOutDelay = VectorMultiplyAdd(VInDelay, VG, VInSamples);
+			
+			// Underflow clamp
+			VectorRegister4Float Mask = VectorBitwiseAnd(VectorCompareGT(VOutDelay, VNFMIN), VectorCompareLT(VOutDelay, VFMIN));
+			VOutDelay = VectorSelect(Mask, GlobalVectorConstants::FloatZero, VOutDelay);
+			VectorStoreAligned(VOutDelay, &OutDelaySamples[i]);
+
+			// y[n] = -G * w[n] + w[n - D]
+			VectorRegister4Float VOut = VectorMultiplyAdd(VOutDelay, VNG, VInDelay);
+			VectorStoreAligned(VOut, &OutSamples[i]);
+		}
+
+		// Calculate allpass for remaining samples that we couldn't SIMD
+		for (int32 i = NumToSIMD; i < InNum; i++)
+		{
+			OutDelaySamples[i] = Audio::UnderflowClamp(InDelaySamples[i] * G + InSamples[i]);
+			OutSamples[i] = OutDelaySamples[i] * -G + InDelaySamples[i];
+		}
 	}
 
-	// Calculate allpass for remaining samples that we couldn't SIMD
-	for (int32 i = NumToSIMD; i < InNum; i++)
+	void FLongDelayAPF::Reset() 
 	{
-		OutDelaySamples[i] = Audio::UnderflowClamp(InDelaySamples[i] * G + InSamples[i]);
-		OutSamples[i] = OutDelaySamples[i] * -G + InDelaySamples[i];
+		DelayLine->ClearSamples();
+		DelayLine->AddZeros(NumDelaySamples);
 	}
-}
 
-void FLongDelayAPF::Reset() 
-{
-	DelayLine->ClearSamples();
-	DelayLine->AddZeros(NumDelaySamples);
-}
-
-int32 FLongDelayAPF::GetNumInternalBufferSamples() const
-{
-	return NumInternalBufferSamples;
+	int32 FLongDelayAPF::GetNumInternalBufferSamples() const
+	{
+		return NumInternalBufferSamples;
+	}
 }
-
diff --git a/Engine/Source/Runtime/SignalProcessing/Private/MultithreadedPatching.cpp b/Engine/Source/Runtime/SignalProcessing/Private/MultithreadedPatching.cpp
--- a/Engine/Source/Runtime/SignalProcessing/Private/MultithreadedPatching.cpp
+++ b/Engine/Source/Runtime/SignalProcessing/Private/MultithreadedPatching.cpp
@@ -1,9 +1,9 @@
 // Copyright Epic Games, Inc. All Rights Reserved.
 
 #include "DSP/MultithreadedPatching.h"
-#include "DSP/BufferVectorOperations.h"
-#include "HAL/IConsoleManager.h"
 
+#include "DSP/FloatArrayMath.h"
+#include "HAL/IConsoleManager.h"
 
 static int32 MultithreadedPatchingPushCallsPerOutputCleanupCheckCVar = 256;
 FAutoConsoleVariableRef CVarMultithreadedPatchingPushCallsPerOutputCleanupCheck(
diff --git a/Engine/Source/Runtime/SignalProcessing/Private/ReverbFast.cpp b/Engine/Source/Runtime/SignalProcessing/Private/ReverbFast.cpp
--- a/Engine/Source/Runtime/SignalProcessing/Private/ReverbFast.cpp
+++ b/Engine/Source/Runtime/SignalProcessing/Private/ReverbFast.cpp
@@ -1,7 +1,7 @@
 // Copyright Epic Games, Inc. All Rights Reserved.
 
 #include "DSP/ReverbFast.h"
-#include "DSP/BufferVectorOperations.h"
+#include "DSP/FloatArrayMath.h"
 #include "AudioMixer.h"
 
 namespace Audio {
diff --git a/Engine/Source/Runtime/SignalProcessing/Private/RingModulation.cpp b/Engine/Source/Runtime/SignalProcessing/Private/RingModulation.cpp
--- a/Engine/Source/Runtime/SignalProcessing/Private/RingModulation.cpp
+++ b/Engine/Source/Runtime/SignalProcessing/Private/RingModulation.cpp
@@ -2,7 +2,7 @@
 
 #include "DSP/RingModulation.h"
 #include "DSP/Dsp.h"
-#include "DSP/BufferVectorOperations.h"
+#include "DSP/FloatArrayMath.h"
 
 namespace Audio
 {
diff --git a/Engine/Source/Runtime/SignalProcessing/Private/SpectrumAnalyzer.cpp b/Engine/Source/Runtime/SignalProcessing/Private/SpectrumAnalyzer.cpp
--- a/Engine/Source/Runtime/SignalProcessing/Private/SpectrumAnalyzer.cpp
+++ b/Engine/Source/Runtime/SignalProcessing/Private/SpectrumAnalyzer.cpp
@@ -1,11 +1,13 @@
 // Copyright Epic Games, Inc. All Rights Reserved.
 
 #include "DSP/SpectrumAnalyzer.h"
-#include "DSP/FFTAlgorithm.h"
+
+#include "Algo/MaxElement.h"
+#include "Algo/MinElement.h"
 #include "DSP/ConstantQ.h"
+#include "DSP/FFTAlgorithm.h"
+#include "DSP/FloatArrayMath.h"
 #include "SignalProcessingModule.h"
-#include "Algo/MinElement.h"
-#include "Algo/MaxElement.h"
 
 namespace Audio
 {
diff --git a/Engine/Source/Runtime/SignalProcessing/Private/UniformPartitionConvolution.cpp b/Engine/Source/Runtime/SignalProcessing/Private/UniformPartitionConvolution.cpp
--- a/Engine/Source/Runtime/SignalProcessing/Private/UniformPartitionConvolution.cpp
+++ b/Engine/Source/Runtime/SignalProcessing/Private/UniformPartitionConvolution.cpp
@@ -279,12 +279,12 @@ namespace Audio
 			// B1r B1i B2r B2i
 
 			// VectorA = A1r A1i A2r A2i
-			VectorRegister4Float VectorInA = VectorLoadAligned(&InAData[i]);
+			VectorRegister4Float VectorInA = VectorLoad(&InAData[i]);
 			// Temp12 = A1i A1r A2i A2r
 			VectorRegister4Float Temp1 = VectorSwizzle(VectorInA, 1, 0, 3, 2);
 
 			// VectorB = B1r B1i B2r B2i
-			VectorRegister4Float VectorInB = VectorLoadAligned(&InBData[i]);
+			VectorRegister4Float VectorInB = VectorLoad(&InBData[i]);
 			// Temp2 = B1r B1r B2r B2r
 			VectorRegister4Float Temp2 = VectorSwizzle(VectorInB, 0, 0, 2, 2);
 			// Temp3 = B1i B1i B2i B2i
@@ -302,10 +302,10 @@ namespace Audio
 			Temp1 = VectorMultiplyAdd(Temp1, SignFlip, VectorInA);
 
 			// VectorOut = O1r + A1rB1r - A1iB1i, O1i + A1iB1r + A1rB1i, O2r + A2rB2r - A2iB2i, O2i + A2iB2r + A2rB2i
-			VectorRegister4Float VectorOut = VectorLoadAligned(&OutData[i]);
+			VectorRegister4Float VectorOut = VectorLoad(&OutData[i]);
 			VectorOut = VectorAdd(Temp1, VectorOut);
 
-			VectorStoreAligned(VectorOut, &OutData[i]);
+			VectorStore(VectorOut, &OutData[i]);
 		}
 
 		for (int32 i = NumSimd; i < Num; i += 2)
@@ -332,11 +332,11 @@ namespace Audio
 
 		for (int32 i = 0; i < NumSimd; i += 4)
 		{
-			VectorRegister4Float VectorIn = VectorLoadAligned(&InData[i]);
+			VectorRegister4Float VectorIn = VectorLoad(&InData[i]);
 
 			VectorRegister4Float VectorOut = VectorMultiply(VectorIn, VectorConstant);
 			
-			VectorStoreAligned(VectorOut, &OutData[i]);
+			VectorStore(VectorOut, &OutData[i]);
 		}
 
 		for (int32 i = NumSimd; i < Num; i++)
diff --git a/Engine/Source/Runtime/SignalProcessing/Private/VectorFFT.cpp b/Engine/Source/Runtime/SignalProcessing/Private/VectorFFT.cpp
--- a/Engine/Source/Runtime/SignalProcessing/Private/VectorFFT.cpp
+++ b/Engine/Source/Runtime/SignalProcessing/Private/VectorFFT.cpp
@@ -2,10 +2,11 @@
 
 #include "VectorFFT.h"
 #include "SignalProcessingModule.h"
-#include "CoreMinimal.h"
 
+#include "Templates/UniquePtr.h"
 #include "DSP/FFTAlgorithm.h"
-#include "DSP/BufferVectorOperations.h"
+#include "DSP/FloatArrayMath.h"
+
 
 namespace Audio
 {
@@ -52,9 +53,6 @@ namespace Audio
 			// @param OutComplex - Interleaved complex data with (2 * FFTSize) num floats.
 			void ForwardComplexToComplex(const float* RESTRICT InComplex, float* RESTRICT OutComplex)
 			{
-				checkf(IsAligned<const float*>(InComplex, AUDIO_SIMD_BYTE_ALIGNMENT), TEXT("Memory must be aligned."));
-				checkf(IsAligned<float*>(OutComplex, AUDIO_SIMD_BYTE_ALIGNMENT), TEXT("Memory must be aligned."));
-
 				// To perform FFT, must complete Log2FFTSize stages.  Each radix pass performs 2^m stages
 				// where 2^m is the radix number. So a radix-4 stage is radix-2^m or radix-2^2. Hence radix
 				// 4 performs two stages. Radix-8 is Radix-2^3, so it performs 3 stages.
@@ -108,9 +106,6 @@ namespace Audio
 
 			void InverseComplexToComplex(const float* RESTRICT InComplex, float* RESTRICT OutComplex)
 			{
-				checkf(IsAligned<const float*>(InComplex, AUDIO_SIMD_BYTE_ALIGNMENT), TEXT("Memory must be aligned."));
-				checkf(IsAligned<float*>(OutComplex, AUDIO_SIMD_BYTE_ALIGNMENT), TEXT("Memory must be aligned."));
-
 				// Perform inverse FFT by complex conjugating the input and output.
 				
 				float* WorkData = InverseWorkBuffer.GetData();
@@ -231,9 +226,9 @@ namespace Audio
 				// Perform operation using SIMD
 				for (int32 i = 0; i < NumToSimd; i += 4)
 				{
-					VectorRegister4Float Value = VectorLoadAligned(&InValues[i]);
+					VectorRegister4Float Value = VectorLoad(&InValues[i]);
 					Value = VectorMultiply(SignFlipImag, Value);
-					VectorStoreAligned(Value, &OutValues[i]);
+					VectorStore(Value, &OutValues[i]);
 				}
 
 				// Perform operation where SIMD not possible.
@@ -273,10 +268,10 @@ namespace Audio
 					const int32 Pos2 = Offset2 + Pos;
 					const int32 Pos3 = Offset3 + Pos;
 
-					VectorRegister4Float A0 = VectorLoadAligned(&InValues[Pos0]);
-					VectorRegister4Float A1 = VectorLoadAligned(&InValues[Pos1]);
-					VectorRegister4Float A2 = VectorLoadAligned(&InValues[Pos2]);
-					VectorRegister4Float A3 = VectorLoadAligned(&InValues[Pos3]);
+					VectorRegister4Float A0 = VectorLoad(&InValues[Pos0]);
+					VectorRegister4Float A1 = VectorLoad(&InValues[Pos1]);
+					VectorRegister4Float A2 = VectorLoad(&InValues[Pos2]);
+					VectorRegister4Float A3 = VectorLoad(&InValues[Pos3]);
 
 					VectorRegister4Float C0 = VectorAdd(A0, A2);
 					VectorRegister4Float C2 = VectorSubtract(A0, A2);
@@ -291,10 +286,10 @@ namespace Audio
 					VectorRegister4Float D2 = VectorAdd(C2, C3ConjSwizzle);
 					VectorRegister4Float D3 = VectorSubtract(C2, C3ConjSwizzle);
 
-					VectorStoreAligned(D0, &OutValues[Pos0]);
-					VectorStoreAligned(D1, &OutValues[Pos1]);
-					VectorStoreAligned(D2, &OutValues[Pos2]);
-					VectorStoreAligned(D3, &OutValues[Pos3]);
+					VectorStore(D0, &OutValues[Pos0]);
+					VectorStore(D1, &OutValues[Pos1]);
+					VectorStore(D2, &OutValues[Pos2]);
+					VectorStore(D3, &OutValues[Pos3]);
 				}
 			}
 
@@ -340,14 +335,14 @@ namespace Audio
 					const int32 Pos6 = Offset6 + Pos;
 					const int32 Pos7 = Offset7 + Pos;
 
-					VectorRegister4Float A0 = VectorLoadAligned(&InValues[Pos0]);
-					VectorRegister4Float A1 = VectorLoadAligned(&InValues[Pos1]);
-					VectorRegister4Float A2 = VectorLoadAligned(&InValues[Pos2]);
-					VectorRegister4Float A3 = VectorLoadAligned(&InValues[Pos3]);
-					VectorRegister4Float A4 = VectorLoadAligned(&InValues[Pos4]);
-					VectorRegister4Float A5 = VectorLoadAligned(&InValues[Pos5]);
-					VectorRegister4Float A6 = VectorLoadAligned(&InValues[Pos6]);
-					VectorRegister4Float A7 = VectorLoadAligned(&InValues[Pos7]);
+					VectorRegister4Float A0 = VectorLoad(&InValues[Pos0]);
+					VectorRegister4Float A1 = VectorLoad(&InValues[Pos1]);
+					VectorRegister4Float A2 = VectorLoad(&InValues[Pos2]);
+					VectorRegister4Float A3 = VectorLoad(&InValues[Pos3]);
+					VectorRegister4Float A4 = VectorLoad(&InValues[Pos4]);
+					VectorRegister4Float A5 = VectorLoad(&InValues[Pos5]);
+					VectorRegister4Float A6 = VectorLoad(&InValues[Pos6]);
+					VectorRegister4Float A7 = VectorLoad(&InValues[Pos7]);
 
 					VectorRegister4Float B0 = VectorAdd(A0, A4);
 					VectorRegister4Float B1 = VectorAdd(A1, A5);
@@ -394,14 +389,14 @@ namespace Audio
 					VectorRegister4Float D6 = VectorMultiplyAdd(VectorNegSqrt2D2, T7Conj, C6);
 					VectorRegister4Float D7 = VectorMultiplyAdd(VectorSqrt2D2, T7Conj, C6);
 
-					VectorStoreAligned(D0, &OutValues[Pos0]);
-					VectorStoreAligned(D1, &OutValues[Pos1]);
-					VectorStoreAligned(D2, &OutValues[Pos2]);
-					VectorStoreAligned(D3, &OutValues[Pos3]);
-					VectorStoreAligned(D4, &OutValues[Pos4]);
-					VectorStoreAligned(D5, &OutValues[Pos5]);
-					VectorStoreAligned(D6, &OutValues[Pos6]);
-					VectorStoreAligned(D7, &OutValues[Pos7]);
+					VectorStore(D0, &OutValues[Pos0]);
+					VectorStore(D1, &OutValues[Pos1]);
+					VectorStore(D2, &OutValues[Pos2]);
+					VectorStore(D3, &OutValues[Pos3]);
+					VectorStore(D4, &OutValues[Pos4]);
+					VectorStore(D5, &OutValues[Pos5]);
+					VectorStore(D6, &OutValues[Pos6]);
+					VectorStore(D7, &OutValues[Pos7]);
 				}
 			}
 
@@ -414,18 +409,18 @@ namespace Audio
 				const int32 NumButterflies = 1 << (InStageIndex - 2);
 
 				// Load weights for butterfly
-				const VectorRegister4Float Weight1Real = VectorLoadAligned(Weights.W1R);
-				const VectorRegister4Float Weight1Imag = VectorLoadAligned(Weights.W1I);
-				const VectorRegister4Float Weight2Real = VectorLoadAligned(Weights.W2R);
-				const VectorRegister4Float Weight2Imag = VectorLoadAligned(Weights.W2I);
-				const VectorRegister4Float Weight3Real = VectorLoadAligned(Weights.W3R);
-				const VectorRegister4Float Weight3Imag = VectorLoadAligned(Weights.W3I);
-
-				const VectorRegister4Float Weight3RealNeg = VectorLoadAligned(Weights.W3RNeg);
-				const VectorRegister4Float Weight2RealNeg = VectorLoadAligned(Weights.W2RNeg);
-				const VectorRegister4Float Weight1RealNeg = VectorLoadAligned(Weights.W1RNeg);
-				const VectorRegister4Float Weight1RealD2 = VectorLoadAligned(Weights.W1RD2);
-				const VectorRegister4Float Weight1RealD3 = VectorLoadAligned(Weights.W1RD3);
+				const VectorRegister4Float Weight1Real = VectorLoad(Weights.W1R);
+				const VectorRegister4Float Weight1Imag = VectorLoad(Weights.W1I);
+				const VectorRegister4Float Weight2Real = VectorLoad(Weights.W2R);
+				const VectorRegister4Float Weight2Imag = VectorLoad(Weights.W2I);
+				const VectorRegister4Float Weight3Real = VectorLoad(Weights.W3R);
+				const VectorRegister4Float Weight3Imag = VectorLoad(Weights.W3I);
+
+				const VectorRegister4Float Weight3RealNeg = VectorLoad(Weights.W3RNeg);
+				const VectorRegister4Float Weight2RealNeg = VectorLoad(Weights.W2RNeg);
+				const VectorRegister4Float Weight1RealNeg = VectorLoad(Weights.W1RNeg);
+				const VectorRegister4Float Weight1RealD2 = VectorLoad(Weights.W1RD2);
+				const VectorRegister4Float Weight1RealD3 = VectorLoad(Weights.W1RD3);
 
 				// Perform butterflies.
 				for (int32 i = 0; i < NumButterflies; i += 2)
@@ -435,10 +430,10 @@ namespace Audio
 					const int32 Pos2 = 2 * (Stride * ButterflyIndex + 2 * NumButterflies + i);
 					const int32 Pos3 = 2 * (Stride * ButterflyIndex + 3 * NumButterflies + i);
 
-					VectorRegister4Float A0 = VectorLoadAligned(&InOutValues[Pos0]);
-					VectorRegister4Float A1 = VectorLoadAligned(&InOutValues[Pos1]);
-					VectorRegister4Float A2 = VectorLoadAligned(&InOutValues[Pos2]);
-					VectorRegister4Float A3 = VectorLoadAligned(&InOutValues[Pos3]);
+					VectorRegister4Float A0 = VectorLoad(&InOutValues[Pos0]);
+					VectorRegister4Float A1 = VectorLoad(&InOutValues[Pos1]);
+					VectorRegister4Float A2 = VectorLoad(&InOutValues[Pos2]);
+					VectorRegister4Float A3 = VectorLoad(&InOutValues[Pos3]);
 
 					VectorRegister4Float A1Swizzle = VectorSwizzle(A1, 1, 0, 3, 2);
 					VectorRegister4Float A2Swizzle = VectorSwizzle(A2, 1, 0, 3, 2);
@@ -460,10 +455,10 @@ namespace Audio
 					VectorRegister4Float D2 = VectorMultiplyAdd(C3Swizzle, Weight1RealD2, C2);
 					VectorRegister4Float D3 = VectorMultiplyAdd(C3Swizzle, Weight1RealD3, C2);
 
-					VectorStoreAligned(D0, &InOutValues[Pos0]);
-					VectorStoreAligned(D1, &InOutValues[Pos1]);
-					VectorStoreAligned(D2, &InOutValues[Pos2]);
-					VectorStoreAligned(D3, &InOutValues[Pos3]);
+					VectorStore(D0, &InOutValues[Pos0]);
+					VectorStore(D1, &InOutValues[Pos1]);
+					VectorStore(D2, &InOutValues[Pos2]);
+					VectorStore(D3, &InOutValues[Pos3]);
 				}
 			}
 
@@ -481,18 +476,18 @@ namespace Audio
 					// Load values for current weight.
 					const FRadix4Weight& Weights = Radix4Weights[i];
 
-					const VectorRegister4Float Weight1Real = VectorLoadAligned(Weights.W1R);
-					const VectorRegister4Float Weight1Imag = VectorLoadAligned(Weights.W1I);
-					const VectorRegister4Float Weight2Real = VectorLoadAligned(Weights.W2R);
-					const VectorRegister4Float Weight2Imag = VectorLoadAligned(Weights.W2I);
-					const VectorRegister4Float Weight3Real = VectorLoadAligned(Weights.W3R);
-					const VectorRegister4Float Weight3Imag = VectorLoadAligned(Weights.W3I);
+					const VectorRegister4Float Weight1Real = VectorLoad(Weights.W1R);
+					const VectorRegister4Float Weight1Imag = VectorLoad(Weights.W1I);
+					const VectorRegister4Float Weight2Real = VectorLoad(Weights.W2R);
+					const VectorRegister4Float Weight2Imag = VectorLoad(Weights.W2I);
+					const VectorRegister4Float Weight3Real = VectorLoad(Weights.W3R);
+					const VectorRegister4Float Weight3Imag = VectorLoad(Weights.W3I);
 
-					const VectorRegister4Float Weight3RealNeg = VectorLoadAligned(Weights.W3RNeg);
-					const VectorRegister4Float Weight2RealNeg = VectorLoadAligned(Weights.W2RNeg);
-					const VectorRegister4Float Weight1RealNeg = VectorLoadAligned(Weights.W1RNeg);
-					const VectorRegister4Float Weight1RealD2 = VectorLoadAligned(Weights.W1RD2);
-					const VectorRegister4Float Weight1RealD3 = VectorLoadAligned(Weights.W1RD3);
+					const VectorRegister4Float Weight3RealNeg = VectorLoad(Weights.W3RNeg);
+					const VectorRegister4Float Weight2RealNeg = VectorLoad(Weights.W2RNeg);
+					const VectorRegister4Float Weight1RealNeg = VectorLoad(Weights.W1RNeg);
+					const VectorRegister4Float Weight1RealD2 = VectorLoad(Weights.W1RD2);
+					const VectorRegister4Float Weight1RealD3 = VectorLoad(Weights.W1RD3);
 
 					for (int32 j = 0; j < 4 ; j += 2)
 					{
@@ -501,10 +496,10 @@ namespace Audio
 						const int32 Pos2 = 2 * (Stride * i + 8 + j);
 						const int32 Pos3 = 2 * (Stride * i + 12 + j);
 
-						VectorRegister4Float A0 = VectorLoadAligned(&InOutValues[Pos0]);
-						VectorRegister4Float A1 = VectorLoadAligned(&InOutValues[Pos1]);
-						VectorRegister4Float A2 = VectorLoadAligned(&InOutValues[Pos2]);
-						VectorRegister4Float A3 = VectorLoadAligned(&InOutValues[Pos3]);
+						VectorRegister4Float A0 = VectorLoad(&InOutValues[Pos0]);
+						VectorRegister4Float A1 = VectorLoad(&InOutValues[Pos1]);
+						VectorRegister4Float A2 = VectorLoad(&InOutValues[Pos2]);
+						VectorRegister4Float A3 = VectorLoad(&InOutValues[Pos3]);
 
 						VectorRegister4Float A1Swizzle = VectorSwizzle(A1, 1, 0, 3, 2);
 						VectorRegister4Float A2Swizzle = VectorSwizzle(A2, 1, 0, 3, 2);
@@ -525,10 +520,10 @@ namespace Audio
 						VectorRegister4Float D2 = VectorMultiplyAdd(C3Swizzle, Weight1RealD2, C2);
 						VectorRegister4Float D3 = VectorMultiplyAdd(C3Swizzle, Weight1RealD3, C2);
 
-						VectorStoreAligned(D0, &InOutValues[Pos0]);
-						VectorStoreAligned(D1, &InOutValues[Pos1]);
-						VectorStoreAligned(D2, &InOutValues[Pos2]);
-						VectorStoreAligned(D3, &InOutValues[Pos3]);
+						VectorStore(D0, &InOutValues[Pos0]);
+						VectorStore(D1, &InOutValues[Pos1]);
+						VectorStore(D2, &InOutValues[Pos2]);
+						VectorStore(D3, &InOutValues[Pos3]);
 					}
 				}
 			}
@@ -546,14 +541,14 @@ namespace Audio
 				const int32 Pos6 = 2 * (3 * InNumButterflies + 4 * InReadIndex);
 				const int32 Pos7 = Pos6 + 4;
 				
-				VectorRegister4Float T0 = VectorLoadAligned(&InValues[Pos0]);
-				VectorRegister4Float T1 = VectorLoadAligned(&InValues[Pos1]);
-				VectorRegister4Float T2 = VectorLoadAligned(&InValues[Pos2]);
-				VectorRegister4Float T3 = VectorLoadAligned(&InValues[Pos3]);
-				VectorRegister4Float T4 = VectorLoadAligned(&InValues[Pos4]);
-				VectorRegister4Float T5 = VectorLoadAligned(&InValues[Pos5]);
-				VectorRegister4Float T6 = VectorLoadAligned(&InValues[Pos6]);
-				VectorRegister4Float T7 = VectorLoadAligned(&InValues[Pos7]);
+				VectorRegister4Float T0 = VectorLoad(&InValues[Pos0]);
+				VectorRegister4Float T1 = VectorLoad(&InValues[Pos1]);
+				VectorRegister4Float T2 = VectorLoad(&InValues[Pos2]);
+				VectorRegister4Float T3 = VectorLoad(&InValues[Pos3]);
+				VectorRegister4Float T4 = VectorLoad(&InValues[Pos4]);
+				VectorRegister4Float T5 = VectorLoad(&InValues[Pos5]);
+				VectorRegister4Float T6 = VectorLoad(&InValues[Pos6]);
+				VectorRegister4Float T7 = VectorLoad(&InValues[Pos7]);
 
 				OutValues.A0 = VectorShuffle(T0, T2, 0, 1, 0, 1);
 				OutValues.A1 = VectorShuffle(T4, T6, 0, 1, 0, 1);
@@ -577,14 +572,14 @@ namespace Audio
 				const int32 Pos6 = 2 * (3 * InNumButterflies + 4 * InWriteIndex);
 				const int32 Pos7 = Pos6 + 4;
 
-				VectorStoreAligned(InResult.D0, &OutValues[Pos0]);
-				VectorStoreAligned(InResult.D1, &OutValues[Pos1]);
-				VectorStoreAligned(InResult.D2, &OutValues[Pos2]);
-				VectorStoreAligned(InResult.D3, &OutValues[Pos3]);
-				VectorStoreAligned(InResult.D4, &OutValues[Pos4]);
-				VectorStoreAligned(InResult.D5, &OutValues[Pos5]);
-				VectorStoreAligned(InResult.D6, &OutValues[Pos6]);
-				VectorStoreAligned(InResult.D7, &OutValues[Pos7]);
+				VectorStore(InResult.D0, &OutValues[Pos0]);
+				VectorStore(InResult.D1, &OutValues[Pos1]);
+				VectorStore(InResult.D2, &OutValues[Pos2]);
+				VectorStore(InResult.D3, &OutValues[Pos3]);
+				VectorStore(InResult.D4, &OutValues[Pos4]);
+				VectorStore(InResult.D5, &OutValues[Pos5]);
+				VectorStore(InResult.D6, &OutValues[Pos6]);
+				VectorStore(InResult.D7, &OutValues[Pos7]);
 			}
 
 
@@ -592,30 +587,30 @@ namespace Audio
 			void Radix4ButterflyFinalIteration(const FFinalInputs& Inputs, const FFinalWeights& InWeights, FFinalOutputs& Outputs)
 			{
 				// Note: Some weights are altered to bake in sign flips to avoid an extra multiply later on.
-				const VectorRegister4Float W2I = VectorLoadAligned(InWeights.W2I);
-				const VectorRegister4Float W2R = VectorLoadAligned(InWeights.W2R);
-				const VectorRegister4Float W3I = VectorLoadAligned(InWeights.W3I);
-				const VectorRegister4Float W3R = VectorLoadAligned(InWeights.W3R);
-				const VectorRegister4Float W4I = VectorLoadAligned(InWeights.W4I);
-				const VectorRegister4Float W4R = VectorLoadAligned(InWeights.W4R);
-				const VectorRegister4Float W5I = VectorLoadAligned(InWeights.W5I);
-				const VectorRegister4Float W5R = VectorLoadAligned(InWeights.W5R);
-				const VectorRegister4Float W6I = VectorLoadAligned(InWeights.W6I);
-				const VectorRegister4Float W6R = VectorLoadAligned(InWeights.W6R);
-				const VectorRegister4Float W7I = VectorLoadAligned(InWeights.W7I);
-				const VectorRegister4Float W7R = VectorLoadAligned(InWeights.W7R);
-
-				const VectorRegister4Float W2RNeg = VectorLoadAligned(InWeights.W2RNeg);
-				const VectorRegister4Float W3RNeg = VectorLoadAligned(InWeights.W3RNeg);
-				const VectorRegister4Float W4RNeg = VectorLoadAligned(InWeights.W4RNeg);
-				const VectorRegister4Float W5RNeg = VectorLoadAligned(InWeights.W5RNeg);
-				const VectorRegister4Float W6RNeg = VectorLoadAligned(InWeights.W6RNeg);
-				const VectorRegister4Float W7RNeg = VectorLoadAligned(InWeights.W7RNeg);
-
-				const VectorRegister4Float W2RD4 = VectorLoadAligned(InWeights.W2RD4);
-				const VectorRegister4Float W2RD6 = VectorLoadAligned(InWeights.W2RD6);
-				const VectorRegister4Float W3RD5 = VectorLoadAligned(InWeights.W3RD5);
-				const VectorRegister4Float W3RD7 = VectorLoadAligned(InWeights.W3RD7);
+				const VectorRegister4Float W2I = VectorLoad(InWeights.W2I);
+				const VectorRegister4Float W2R = VectorLoad(InWeights.W2R);
+				const VectorRegister4Float W3I = VectorLoad(InWeights.W3I);
+				const VectorRegister4Float W3R = VectorLoad(InWeights.W3R);
+				const VectorRegister4Float W4I = VectorLoad(InWeights.W4I);
+				const VectorRegister4Float W4R = VectorLoad(InWeights.W4R);
+				const VectorRegister4Float W5I = VectorLoad(InWeights.W5I);
+				const VectorRegister4Float W5R = VectorLoad(InWeights.W5R);
+				const VectorRegister4Float W6I = VectorLoad(InWeights.W6I);
+				const VectorRegister4Float W6R = VectorLoad(InWeights.W6R);
+				const VectorRegister4Float W7I = VectorLoad(InWeights.W7I);
+				const VectorRegister4Float W7R = VectorLoad(InWeights.W7R);
+
+				const VectorRegister4Float W2RNeg = VectorLoad(InWeights.W2RNeg);
+				const VectorRegister4Float W3RNeg = VectorLoad(InWeights.W3RNeg);
+				const VectorRegister4Float W4RNeg = VectorLoad(InWeights.W4RNeg);
+				const VectorRegister4Float W5RNeg = VectorLoad(InWeights.W5RNeg);
+				const VectorRegister4Float W6RNeg = VectorLoad(InWeights.W6RNeg);
+				const VectorRegister4Float W7RNeg = VectorLoad(InWeights.W7RNeg);
+
+				const VectorRegister4Float W2RD4 = VectorLoad(InWeights.W2RD4);
+				const VectorRegister4Float W2RD6 = VectorLoad(InWeights.W2RD6);
+				const VectorRegister4Float W3RD5 = VectorLoad(InWeights.W3RD5);
+				const VectorRegister4Float W3RD7 = VectorLoad(InWeights.W3RD7);
 
 				VectorRegister4Float A2Swizzle = VectorSwizzle(Inputs.A2, 1, 0, 3, 2);
 				VectorRegister4Float A3Swizzle = VectorSwizzle(Inputs.A3, 1, 0, 3, 2);
@@ -1076,23 +1071,23 @@ namespace Audio
 
 		if (FFTSize > InStartIndex)
 		{
-			VectorRegister4Float VInRev1 = VectorLoadAligned(&InValues[FFTSize - InStartIndex]);
+			VectorRegister4Float VInRev1 = VectorLoad(&InValues[FFTSize - InStartIndex]);
 
 			for (int32 i = InStartIndex; i < FFTSize; i += 4)
 			{
-				VectorRegister4Float VIn = VectorLoadAligned(&InValues[i]);
+				VectorRegister4Float VIn = VectorLoad(&InValues[i]);
 				VectorRegister4Float VInRISwap = VectorSwizzle(VIn, 1, 0, 3, 2);
 
-				VectorRegister4Float VInRev2 = VectorLoadAligned(&InValues[FFTSize - i - 4]);
+				VectorRegister4Float VInRev2 = VectorLoad(&InValues[FFTSize - i - 4]);
 				VectorRegister4Float VInRev = VectorShuffle(VInRev1, VInRev2, 0, 1, 2, 3);
 				VInRev1 = VInRev2;
 
 				VectorRegister4Float VInRevRISwap = VectorSwizzle(VInRev, 1, 0, 3, 2);
 
-				VectorRegister4Float VAlphaReal = VectorLoadAligned(&AlphaRealData[i]);
-				VectorRegister4Float VAlphaImag = VectorLoadAligned(&AlphaImagData[i]);
-				VectorRegister4Float VBetaReal = VectorLoadAligned(&BetaRealData[i]);
-				VectorRegister4Float VBetaImag = VectorLoadAligned(&BetaImagData[i]);
+				VectorRegister4Float VAlphaReal = VectorLoad(&AlphaRealData[i]);
+				VectorRegister4Float VAlphaImag = VectorLoad(&AlphaImagData[i]);
+				VectorRegister4Float VBetaReal = VectorLoad(&BetaRealData[i]);
+				VectorRegister4Float VBetaImag = VectorLoad(&BetaImagData[i]);
 
 				// Out1 = [ R * Ar,  I * Ar]
 				// Out2 = [ I * Ai,  R * Ai]
@@ -1113,7 +1108,7 @@ namespace Audio
 				// 	(I * Ar) + (R * Ai) + (NR * Bi) + (NI * Br)
 				// ]
 				VectorRegister4Float Out = VectorAdd(Out12, Out34);
-				VectorStoreAligned(Out, &OutValues[i]);
+				VectorStore(Out, &OutValues[i]);
 			}
 		}
 	}
@@ -1151,9 +1146,6 @@ namespace Audio
 
 	void FVectorRealToComplexFFT::ForwardRealToComplex(const float* RESTRICT InReal, float* RESTRICT OutComplex)
 	{
-		checkf(IsAligned<const float*>(InReal, AUDIO_SIMD_BYTE_ALIGNMENT), TEXT("Memory must be aligned."));
-		checkf(IsAligned<float*>(OutComplex, AUDIO_SIMD_BYTE_ALIGNMENT), TEXT("Memory must be aligned."));
-
 		// Performs a N sized real-to-complex FFT using an N/2 complex-to-complex FFT.
 		float* WorkData = WorkBuffer.GetData();
 
@@ -1196,9 +1188,6 @@ namespace Audio
 
 	void FVectorRealToComplexFFT::InverseComplexToReal(const float* RESTRICT InComplex, float* RESTRICT OutReal)
 	{
-		checkf(IsAligned<const float*>(InComplex, AUDIO_SIMD_BYTE_ALIGNMENT), TEXT("Memory must be aligned."));
-		checkf(IsAligned<float*>(OutReal, AUDIO_SIMD_BYTE_ALIGNMENT), TEXT("Memory must be aligned."));
-
 		// Performs a N sized complex-to-real FFT using an N/2 complex-to-complex FFT.
 
 		float* WorkData = WorkBuffer.GetData();
@@ -1275,7 +1264,7 @@ namespace Audio
 	/** If true, this implementation requires input and output arrays to be 128 bit aligned. */
 	bool FVectorFFTFactory::Expects128BitAlignedArrays() const
 	{
-		return true;
+		return false;
 	}
 
 	/** Returns true if the input settings are supported by this factory. */
diff --git a/Engine/Source/Runtime/SignalProcessing/Private/VoiceProcessing.cpp b/Engine/Source/Runtime/SignalProcessing/Private/VoiceProcessing.cpp
--- a/Engine/Source/Runtime/SignalProcessing/Private/VoiceProcessing.cpp
+++ b/Engine/Source/Runtime/SignalProcessing/Private/VoiceProcessing.cpp
@@ -1,7 +1,9 @@
 // Copyright Epic Games, Inc. All Rights Reserved.
 
 #include "DSP/VoiceProcessing.h"
+
 #include "DSP/AudioDebuggingUtilities.h"
+#include "DSP/FloatArrayMath.h"
 
 namespace Audio
 {
@@ -66,7 +68,7 @@ namespace Audio
 		AccumulatedSum = VectorAdd(AccumulatedSum, TotalAccumulation);
 
 		alignas(16) float PartionedSums[4];
-		VectorStoreAligned(AccumulatedSum, PartionedSums);
+		VectorStore(AccumulatedSum, PartionedSums);
 
 		return (PartionedSums[0] + PartionedSums[1] + PartionedSums[2] + PartionedSums[3]) / (AudioBuffer.Num() * 4);
 	}
@@ -93,7 +95,7 @@ namespace Audio
 			// stop outputting audio then.
 			for (int32 InSampleIndex = 0; InSampleIndex < NumSamples; InSampleIndex += 4)
 			{
-				const VectorRegister4Float InputVector = VectorLoadAligned(&InAudio[InSampleIndex]);
+				const VectorRegister4Float InputVector = VectorLoad(&InAudio[InSampleIndex]);
 				VectorRegister4Float OutputVector;
 				float InstantaneousAmplitude = Averager.ProcessAudio(InputVector, OutputVector);
 				CurrentAmplitude = ReleaseTau * (CurrentAmplitude - InstantaneousAmplitude) + InstantaneousAmplitude;
@@ -122,7 +124,7 @@ namespace Audio
 			float InstantaneousAmplitude = 0.0f;
 			for (int32 InSampleIndex = 0; InSampleIndex < NumSamples; InSampleIndex += 4)
 			{
-				const VectorRegister4Float InputVector = VectorLoadAligned(&InAudio[InSampleIndex]);
+				const VectorRegister4Float InputVector = VectorLoad(&InAudio[InSampleIndex]);
 				VectorRegister4Float OutputVector;
 				InstantaneousAmplitude = Averager.ProcessAudio(InputVector, OutputVector);
 
@@ -238,15 +240,15 @@ namespace Audio
 		
 		for (int32 WeightIndex = 0; WeightIndex < NumWeights; WeightIndex+= 4)
 		{
-			VectorRegister4Float TargetReal = VectorLoadAligned(&InWeightsReal[WeightIndex]);
-			VectorRegister4Float CurrentReal = VectorLoadAligned(&CurrentRealBuffer[WeightIndex]);
+			VectorRegister4Float TargetReal = VectorLoad(&InWeightsReal[WeightIndex]);
+			VectorRegister4Float CurrentReal = VectorLoad(&CurrentRealBuffer[WeightIndex]);
 			const VectorRegister4Float DeltaReal = VectorMultiply(VectorSubtract(TargetReal, CurrentReal), ConvergenceRate);
-			VectorStoreAligned(DeltaReal, &DeltaRealBuffer[WeightIndex]);
+			VectorStore(DeltaReal, &DeltaRealBuffer[WeightIndex]);
 
-			VectorRegister4Float TargetImag = VectorLoadAligned(&InWeightsImag[WeightIndex]);
-			VectorRegister4Float CurrentImag = VectorLoadAligned(&CurrentImagBuffer[WeightIndex]);
+			VectorRegister4Float TargetImag = VectorLoad(&InWeightsImag[WeightIndex]);
+			VectorRegister4Float CurrentImag = VectorLoad(&CurrentImagBuffer[WeightIndex]);
 			const VectorRegister4Float DeltaImag = VectorMultiply(VectorSubtract(TargetImag, CurrentImag), ConvergenceRate);
-			VectorStoreAligned(DeltaImag, &DeltaImagBuffer[WeightIndex]);
+			VectorStore(DeltaImag, &DeltaImagBuffer[WeightIndex]);
 		}
 	}
 
@@ -261,15 +263,15 @@ namespace Audio
 
 		for (int32 WeightIndex = 0; WeightIndex < NumWeights; WeightIndex += 4)
 		{
-			VectorRegister4Float DeltaReal = VectorLoadAligned(&DeltasRealBuffer[WeightIndex]);
-			VectorRegister4Float CurrentReal = VectorLoadAligned(&CurrentRealBuffer[WeightIndex]);
+			VectorRegister4Float DeltaReal = VectorLoad(&DeltasRealBuffer[WeightIndex]);
+			VectorRegister4Float CurrentReal = VectorLoad(&CurrentRealBuffer[WeightIndex]);
 			CurrentReal = VectorAdd(CurrentReal, DeltaReal);
-			VectorStoreAligned(CurrentReal, &CurrentRealBuffer[WeightIndex]);
+			VectorStore(CurrentReal, &CurrentRealBuffer[WeightIndex]);
 
-			VectorRegister4Float DeltaImag = VectorLoadAligned(&DeltasImagBuffer[WeightIndex]);
-			VectorRegister4Float CurrentImag = VectorLoadAligned(&CurrentImagBuffer[WeightIndex]);
+			VectorRegister4Float DeltaImag = VectorLoad(&DeltasImagBuffer[WeightIndex]);
+			VectorRegister4Float CurrentImag = VectorLoad(&CurrentImagBuffer[WeightIndex]);
 			CurrentImag = VectorAdd(CurrentImag, DeltaImag);
-			VectorStoreAligned(CurrentImag, &CurrentImagBuffer[WeightIndex]);
+			VectorStore(CurrentImag, &CurrentImagBuffer[WeightIndex]);
 		}
 	}
 
diff --git a/Engine/Source/Runtime/SignalProcessing/Public/DSP/AlignedBuffer.h b/Engine/Source/Runtime/SignalProcessing/Public/DSP/AlignedBuffer.h
new file mode 100644
--- /dev/null
+++ b/Engine/Source/Runtime/SignalProcessing/Public/DSP/AlignedBuffer.h
@@ -0,0 +1,32 @@
+// Copyright Epic Games, Inc. All Rights Reserved.
+
+#pragma once
+
+#include "Containers/Array.h"
+#include "Containers/ContainerAllocationPolicies.h"
+
+#if PLATFORM_SWITCH
+// Switch uses page alignment for submitted buffers
+#define AUDIO_BUFFER_ALIGNMENT 4096
+#else
+#define AUDIO_BUFFER_ALIGNMENT 16
+#endif
+
+#define AUDIO_SIMD_BYTE_ALIGNMENT (16)
+#define AUDIO_NUM_FLOATS_PER_VECTOR_REGISTER (4)
+
+
+namespace Audio
+{
+	/** Aligned allocator used for fast operations. */	
+	using FAudioBufferAlignedAllocator = TAlignedHeapAllocator<AUDIO_BUFFER_ALIGNMENT>;
+
+	using FAlignedByteBuffer  = TArray<uint8, FAudioBufferAlignedAllocator>;
+	using FAlignedFloatBuffer = TArray<float, FAudioBufferAlignedAllocator>;
+	using FAlignedInt32Buffer = TArray<int32, FAudioBufferAlignedAllocator>;
+
+	// Deprecated in favor of versions above
+	typedef TArray<uint8, FAudioBufferAlignedAllocator> AlignedByteBuffer;
+	typedef TArray<float, FAudioBufferAlignedAllocator> AlignedFloatBuffer;
+	typedef TArray<int32, FAudioBufferAlignedAllocator> AlignedInt32Buffer;
+}
diff --git a/Engine/Source/Runtime/SignalProcessing/Public/DSP/AudioDebuggingUtilities.h b/Engine/Source/Runtime/SignalProcessing/Public/DSP/AudioDebuggingUtilities.h
--- a/Engine/Source/Runtime/SignalProcessing/Public/DSP/AudioDebuggingUtilities.h
+++ b/Engine/Source/Runtime/SignalProcessing/Public/DSP/AudioDebuggingUtilities.h
@@ -2,38 +2,18 @@
 
 #pragma once
 
+#include "HAL/Platform.h"
+
+#if UE_ENABLE_INCLUDE_ORDER_DEPRECATED_IN_5_1 
 #include "CoreMinimal.h"
 #include "Misc/ScopeLock.h"
 #include "BufferVectorOperations.h"
+#endif
 
 #define USE_AUDIO_DEBUGGING UE_BUILD_DEBUG
 
-void BreakWhenAudible(float* InBuffer, int32 NumSamples)
-{
-	static const float AudibilityThreshold = Audio::ConvertToLinear(-40.0f);
-
-	TArrayView<const float> InBufferView(InBuffer, NumSamples);
-	float BufferAmplitude = Audio::ArrayGetAverageAbsValue(InBufferView);
-
-	if (BufferAmplitude > AudibilityThreshold)
-	{
-		PLATFORM_BREAK();
-	}
-}
-
-void BreakWhenTooLoud(float* InBuffer, int32 NumSamples)
-{
-	static const float PainThreshold = Audio::ConvertToLinear(3.0f);
-
-	TArrayView<const float> InBufferView(InBuffer, NumSamples);
-	float BufferAmplitude = Audio::ArrayGetAverageAbsValue(InBufferView);
-
-	if (BufferAmplitude > PainThreshold)
-	{
-		PLATFORM_BREAK();
-	}
-}
-
+SIGNALPROCESSING_API void BreakWhenAudible(float* InBuffer, int32 NumSamples);
+SIGNALPROCESSING_API void BreakWhenTooLoud(float* InBuffer, int32 NumSamples);
 
 #if USE_AUDIO_DEBUGGING
 #define BREAK_WHEN_AUDIBLE(Ptr, Num) BreakWhenAudible(Ptr, Num);
@@ -41,4 +21,4 @@ void BreakWhenTooLoud(float* InBuffer, int32 NumSamples)
 #else
 #define BREAK_WHEN_AUDIBLE(Ptr, Num) 
 #define BREAK_WHEN_TOO_LOUD(Ptr, Num)
-#endif
\ No newline at end of file
+#endif
diff --git a/Engine/Source/Runtime/SignalProcessing/Public/DSP/AudioFFT.h b/Engine/Source/Runtime/SignalProcessing/Public/DSP/AudioFFT.h
--- a/Engine/Source/Runtime/SignalProcessing/Public/DSP/AudioFFT.h
+++ b/Engine/Source/Runtime/SignalProcessing/Public/DSP/AudioFFT.h
@@ -45,6 +45,8 @@ namespace Audio
 	class SIGNALPROCESSING_API FWindow
 	{
 	public:
+		FWindow() = delete;
+
 		/**
 		 * Constructor. Allocates buffer and generates window inside of it.
 		 * @param InType: The type of window that should be generated.
@@ -54,82 +56,21 @@ namespace Audio
 		 *                     Generally, set this to false if using this window with an STFT, but use true
 		 *                     if this window will be used on an entire, self-contained signal.
 		 */
-		FWindow(EWindowType InType, int32 InNumFrames, int32 InNumChannels, bool bIsPeriodic)
-			: WindowType(InType)
-			, NumSamples(InNumFrames * InNumChannels)
-		{
-			checkf(NumSamples % 4 == 0, TEXT("For performance reasons, this window's length should be a multiple of 4."));
-			Generate(InNumFrames, InNumChannels, bIsPeriodic);
-		}
-
-		// Destructor. Releases memory used for window.
-		~FWindow()
-		{
-		}
+		FWindow(EWindowType InType, int32 InNumFrames, int32 InNumChannels, bool bIsPeriodic);
 
 		// Apply this window to InBuffer, which is expected to be an interleaved buffer with the same amount of frames
 		// and channels this window was constructed with.
-		void ApplyToBuffer(float* InBuffer)
-		{
-			if (WindowType == EWindowType::None)
-			{
-				return;
-			}
-
-			check(IsAligned<float*>(InBuffer, 4));
-			TArrayView<const float> WindowBufferView(WindowBuffer.GetData(), NumSamples);
-			TArrayView<float> InBufferView(InBuffer, NumSamples);
-			ArrayMultiplyInPlace(WindowBufferView, InBufferView);
-		}
+		void ApplyToBuffer(float* InBuffer);
 
-		EWindowType GetWindowType() const
-		{
-			return WindowType;
-		}
+		EWindowType GetWindowType() const;
 
 	private:
 		EWindowType WindowType;
 		FAlignedFloatBuffer WindowBuffer;
 		int32 NumSamples;
 
-		// Purposefully hidden constructor.
-		FWindow();
-
 		// Generate the window. Called on constructor.
-		void Generate(int32 NumFrames, int32 NumChannels, bool bIsPeriodic)
-		{
-			if (WindowType == EWindowType::None)
-			{
-				return;
-			}
-
-			WindowBuffer.Reset();
-			WindowBuffer.AddZeroed(NumSamples);
-
-			switch (WindowType)
-			{
-			case EWindowType::Hann:
-			{
-				GenerateHannWindow(WindowBuffer.GetData(), NumFrames, NumChannels, bIsPeriodic);
-				break;
-			}
-			case EWindowType::Hamming:
-			{
-				GenerateHammingWindow(WindowBuffer.GetData(), NumFrames, NumChannels, bIsPeriodic);
-				break;
-			}
-			case EWindowType::Blackman:
-			{
-				GenerateBlackmanWindow(WindowBuffer.GetData(), NumFrames, NumChannels, bIsPeriodic);
-				break;
-			}
-			default:
-			{
-				checkf(false, TEXT("Unknown window type!"));
-				break;
-			}
-			}
-		}
+		void Generate(int32 NumFrames, int32 NumChannels, bool bIsPeriodic);
 	};
 
 	struct FFTTimeDomainData
diff --git a/Engine/Source/Runtime/SignalProcessing/Public/DSP/BufferVectorOperations.h b/Engine/Source/Runtime/SignalProcessing/Public/DSP/BufferVectorOperations.h
--- a/Engine/Source/Runtime/SignalProcessing/Public/DSP/BufferVectorOperations.h
+++ b/Engine/Source/Runtime/SignalProcessing/Public/DSP/BufferVectorOperations.h
@@ -2,32 +2,15 @@
 
 #pragma once
 
-#include "CoreMinimal.h"
+#include "DSP/AlignedBuffer.h"
+#include "Containers/Array.h"
 
-#if PLATFORM_SWITCH
-// Switch uses page alignment for submitted buffers
-#define AUDIO_BUFFER_ALIGNMENT 4096
-#else
-#define AUDIO_BUFFER_ALIGNMENT 16
+#if UE_ENABLE_INCLUDE_ORDER_DEPRECATED_IN_5_1 
+#include "CoreMinimal.h"
 #endif
 
-#define AUDIO_SIMD_BYTE_ALIGNMENT (16)
-#define AUDIO_NUM_FLOATS_PER_VECTOR_REGISTER (4)
-
 namespace Audio
 {
-	/** Aligned allocator used for fast operations. */	
-	using FAudioBufferAlignedAllocator = TAlignedHeapAllocator<AUDIO_BUFFER_ALIGNMENT>;
-
-	using FAlignedByteBuffer  = TArray<uint8, FAudioBufferAlignedAllocator>;
-	using FAlignedFloatBuffer = TArray<float, FAudioBufferAlignedAllocator>;
-	using FAlignedInt32Buffer = TArray<int32, FAudioBufferAlignedAllocator>;
-
-	// Deprecated in favor of versions above
-	typedef TArray<uint8, FAudioBufferAlignedAllocator> AlignedByteBuffer;
-	typedef TArray<float, FAudioBufferAlignedAllocator> AlignedFloatBuffer;
-	typedef TArray<int32, FAudioBufferAlignedAllocator> AlignedInt32Buffer;
-
 	/** CHANNEL-SPECIFIC OPERATIONS */
 
 	/** Takes a 2 channel interleaved buffer and applies Gains to it. Gains is expected to point to a 2 float long buffer.
diff --git a/Engine/Source/Runtime/SignalProcessing/Public/DSP/ConstantQ.h b/Engine/Source/Runtime/SignalProcessing/Public/DSP/ConstantQ.h
--- a/Engine/Source/Runtime/SignalProcessing/Public/DSP/ConstantQ.h
+++ b/Engine/Source/Runtime/SignalProcessing/Public/DSP/ConstantQ.h
@@ -2,9 +2,13 @@
 
 #pragma once
 
-#include "CoreMinimal.h"
+#include "DSP/AlignedBuffer.h"
 #include "DSP/FloatArrayMath.h"
 
+#if UE_ENABLE_INCLUDE_ORDER_DEPRECATED_IN_5_1
+#include "CoreMinimal.h"
+#endif
+
 namespace Audio
 {
 	/** Band normalization schemes */
@@ -81,6 +85,5 @@ namespace Audio
 
 		// Fills an array with the spectrum weights to apply to an FFT magnitude or power spectrum output.
 		static void FillArrayWithConstantQBand(const FPseudoConstantQBandSettings& InSettings, FAlignedFloatBuffer& OutOffsetArray, int32& OutOffsetIndex);
-
 	};
 }
diff --git a/Engine/Source/Runtime/SignalProcessing/Public/DSP/Dsp.h b/Engine/Source/Runtime/SignalProcessing/Public/DSP/Dsp.h
--- a/Engine/Source/Runtime/SignalProcessing/Public/DSP/Dsp.h
+++ b/Engine/Source/Runtime/SignalProcessing/Public/DSP/Dsp.h
@@ -2,13 +2,22 @@
 
 #pragma once
 
-#include "CoreMinimal.h"
-#include "AudioMixerCore/Public/AudioDefines.h"
+#include "Containers/Array.h"
+#include "DSP/AlignedBuffer.h"
+#include "HAL/Platform.h"
+#include "HAL/UnrealMemory.h"
+#include "Logging/LogMacros.h"
+#include "Math/UnrealMath.h"
 #include "Misc/ScopeLock.h"
+#include "SignalProcessingModule.h"
 #include "Templates/IsFloatingPoint.h"
 #include "Templates/IsIntegral.h"
 #include "Templates/IsSigned.h"
-#include "FloatArrayMath.h"
+
+#if UE_ENABLE_INCLUDE_ORDER_DEPRECATED_IN_5_1 
+#include "AudioMixerCore/Public/AudioDefines.h"
+#include "CoreMinimal.h"
+#endif
 
 // Macros which can be enabled to cause DSP sample checking
 #if 0
@@ -27,7 +36,7 @@ namespace Audio
 	{	
 		if (InSample > Threshold || InSample < -Threshold)
 		{
-			UE_LOG(LogTemp, Log, TEXT("SampleValue Was %.2f"), InSample);
+			UE_LOG(LogSignalProcessing, Log, TEXT("SampleValue Was %.2f"), InSample);
 		}
 	}
 
@@ -181,7 +190,7 @@ namespace Audio
 				PhaseSource[2] = LastPhase + 2 * PhasePerSample;
 				PhaseSource[3] = LastPhase + 3 * PhasePerSample;
 
-				VectorRegister4Float PhaseVec = VectorLoadAligned(PhaseSource);
+				VectorRegister4Float PhaseVec = VectorLoad(PhaseSource);
 				VectorRegister4Float XVector, YVector;
 
 				// We need an accurate representation of the delta
@@ -218,7 +227,6 @@ namespace Audio
 				{
 					// We've actually already calculated the next quad - it's in YVector
 					alignas(16) float YFloats[4];
-					//VectorStoreAligned(YVector, YFloats);
 					VectorStore(YVector, YFloats);
 
 					int32 Remn = BlockSampleCount & SIMD_MASK;
@@ -278,7 +286,6 @@ namespace Audio
 	static void ConvertBipolarBufferToUnipolar(float* InAlignedBuffer, int32 NumSamples)
 	{
 		// Make sure buffers are aligned and we can do a whole number of loops.
-		check(IsAligned(InAlignedBuffer, sizeof(VectorRegister4Float)));
 		check(NumSamples % 4 == 0);
 
 		const VectorRegister4Float Half = VectorSetFloat1(0.5f);
@@ -286,10 +293,10 @@ namespace Audio
 		// Process buffer 1 vector (4 floats) at a time.
 		for(int32 i = NumSamples / 4; i; --i, InAlignedBuffer += 4)
 		{
-			VectorRegister4Float V = VectorLoadAligned(InAlignedBuffer);
+			VectorRegister4Float V = VectorLoad(InAlignedBuffer);
 			V = VectorMultiply(V, Half);
 			V = VectorAdd(V, Half);
-			VectorStoreAligned(V, InAlignedBuffer);
+			VectorStore(V, InAlignedBuffer);
 		}
 	}
 
diff --git a/Engine/Source/Runtime/SignalProcessing/Public/DSP/EnvelopeFollower.h b/Engine/Source/Runtime/SignalProcessing/Public/DSP/EnvelopeFollower.h
--- a/Engine/Source/Runtime/SignalProcessing/Public/DSP/EnvelopeFollower.h
+++ b/Engine/Source/Runtime/SignalProcessing/Public/DSP/EnvelopeFollower.h
@@ -2,9 +2,15 @@
 
 #pragma once
 
-#include "CoreMinimal.h"
+#include "Containers/Array.h"
+#include "DSP/AlignedBuffer.h"
 #include "DSP/Dsp.h"
+#include "HAL/Platform.h"
+
+#if UE_ENABLE_INCLUDE_ORDER_DEPRECATED_IN_5_1 
+#include "CoreMinimal.h"
 #include "DSP/FloatArrayMath.h"
+#endif
 
 namespace Audio
 {
diff --git a/Engine/Source/Runtime/SignalProcessing/Public/DSP/FloatArrayMath.h b/Engine/Source/Runtime/SignalProcessing/Public/DSP/FloatArrayMath.h
--- a/Engine/Source/Runtime/SignalProcessing/Public/DSP/FloatArrayMath.h
+++ b/Engine/Source/Runtime/SignalProcessing/Public/DSP/FloatArrayMath.h
@@ -9,11 +9,9 @@ namespace Audio
 {
 	/** Sum all values in an array. */
 	SIGNALPROCESSING_API void ArraySum(TArrayView<const float> InValues, float& OutSum);
-	SIGNALPROCESSING_API void ArraySum(const FAlignedFloatBuffer& InValues, float& OutSum);
 
 	/** Sums two buffers together and places the result in the resulting buffer. */
 	SIGNALPROCESSING_API void ArraySum(TArrayView<const float> InFloatBuffer1, TArrayView<const float> InFloatBuffer2, TArrayView<float> OutputBuffer);
-	SIGNALPROCESSING_API void ArraySum(const FAlignedFloatBuffer& InFloatBuffer1, const FAlignedFloatBuffer& InFloatBuffer2, FAlignedFloatBuffer& OutputBuffer);
 
 	/** Cumulative sum of array.
 	 *
@@ -38,15 +36,12 @@ namespace Audio
 
 	/** Takes an audio buffer and returns the magnitude across that buffer. */
 	SIGNALPROCESSING_API float ArrayGetMagnitude(TArrayView<const float> Buffer);
-	SIGNALPROCESSING_API float ArrayGetMagnitude(const FAlignedFloatBuffer& Buffer);
 
 	/** Takes an audio buffer and gets the average amplitude across that buffer. */
 	SIGNALPROCESSING_API float ArrayGetAverageValue(TArrayView<const float> Buffer);
-	SIGNALPROCESSING_API float ArrayGetAverageValue(const FAlignedFloatBuffer& Buffer);
 
 	/** Takes an audio buffer and gets the average absolute amplitude across that buffer. */
 	SIGNALPROCESSING_API float ArrayGetAverageAbsValue(TArrayView<const float> Buffer);
-	SIGNALPROCESSING_API float ArrayGetAverageAbsValue(const FAlignedFloatBuffer& Buffer);
 
 	/** Mean filter of array.
 	 *
@@ -117,22 +112,18 @@ namespace Audio
 
 	/** Multiply the second buffer in place by the first buffer. */
 	SIGNALPROCESSING_API void ArrayMultiplyInPlace(TArrayView<const float> InFloatBuffer, TArrayView<float> BufferToMultiply);
-	SIGNALPROCESSING_API void ArrayMultiplyInPlace(const FAlignedFloatBuffer& InFloatBuffer, FAlignedFloatBuffer& BufferToMultiply);
 
 	/** Multiplies two complex valued arrays element-wise. 
 	 * This assumes elements are in interleaved format [real_0, imag_0, ..., real_N, imag_N]
 	 * Stores result in InValues2
 	 */
 	SIGNALPROCESSING_API void ArrayComplexMultiplyInPlace(TArrayView<const float> InValues1, TArrayView<float> InValues2);
-	SIGNALPROCESSING_API void ArrayComplexMultiplyInPlace(const FAlignedFloatBuffer& InValues1, FAlignedFloatBuffer& InValues2);
 
 	/** Multiplies the input float buffer with the given value. */
 	SIGNALPROCESSING_API void ArrayMultiplyByConstant(TArrayView<const float> InFloatBuffer, float InValue, TArrayView<float> OutFloatBuffer);
-	SIGNALPROCESSING_API void ArrayMultiplyByConstant(const FAlignedFloatBuffer& InFloatBuffer, float InValue, FAlignedFloatBuffer& OutFloatBuffer);
 
 	/** Similar to ArrayMultiplyByConstant, but performs the multiply in place. */
 	SIGNALPROCESSING_API void ArrayMultiplyByConstantInPlace(TArrayView<float> InOutBuffer, float InGain);
-	SIGNALPROCESSING_API void ArrayMultiplyByConstantInPlace(FAlignedFloatBuffer& InOutBuffer, float InGain);
 
 	/** Add arrays element-wise in place. InAccumulateValues[i] += InValues[i]
 	 *
@@ -140,11 +131,9 @@ namespace Audio
 	 *  InAccumulateValues is the array which holds the sum.
 	 */
 	SIGNALPROCESSING_API void ArrayAddInPlace(TArrayView<const float> InValues, TArrayView<float> InAccumulateValues);
-	SIGNALPROCESSING_API void ArrayAddInPlace(const FAlignedFloatBuffer& InValues, FAlignedFloatBuffer& InAccumulateValues);
 
 	/** Adds a constant to a buffer (useful for DC offset removal) */
 	SIGNALPROCESSING_API void ArrayAddConstantInplace(TArrayView<float> InOutBuffer, float InConstant);
-	SIGNALPROCESSING_API void ArrayAddConstantInplace(FAlignedFloatBuffer& InOutBuffer, float InConstant);
 
 	/** Multiply Add arrays element-wise in place. InAccumulateValues[i] += InMultiplier * InValues[i]
 	 *
@@ -153,7 +142,6 @@ namespace Audio
 	 *  @param InAccumulateValues - The array which holds the sum.
 	 */
 	SIGNALPROCESSING_API void ArrayMultiplyAddInPlace(TArrayView<const float> InValues, float InMultiplier, TArrayView<float> InAccumulateValues);
-	SIGNALPROCESSING_API void ArrayMultiplyAddInPlace(const FAlignedFloatBuffer& InValues, float InMultiplier, FAlignedFloatBuffer& InAccumulateValues);
 
 	/** Linearly Interpolate Add arrays element-wise in place. InAccumulateValues[i] += ((1 - alpha) * InStartMultiplier + alpha * InEndMultipler) * InValues[i]
 	 * Interpolation is performed over the length of the array.
@@ -164,7 +152,6 @@ namespace Audio
 	 *  @param InAccumulateValues - The array which holds the sum.
 	 */
 	SIGNALPROCESSING_API void ArrayLerpAddInPlace(TArrayView<const float> InValues, float InStartMultiplier, float InEndMultiplier, TArrayView<float> InAccumulateValues);
-	SIGNALPROCESSING_API void ArrayLerpAddInPlace(const FAlignedFloatBuffer& InValues, float InStartMultiplier, float InEndMultiplier, FAlignedFloatBuffer& InAccumulateValues);
 
 	/** Subract arrays element-wise. OutArray = InMinuend - InSubtrahend
 	 *
@@ -173,19 +160,15 @@ namespace Audio
 	 *  OutBuffer is the array which holds the result.
 	 */
 	SIGNALPROCESSING_API void ArraySubtract(TArrayView<const float> InMinuend, TArrayView<const float> InSubtrahend, TArrayView<float> OutBuffer);
-	SIGNALPROCESSING_API void ArraySubtract(const FAlignedFloatBuffer& InMinuend, const FAlignedFloatBuffer& InSubtrahend, FAlignedFloatBuffer& OutputBuffer);
 
 	/* Performs element-wise in-place subtraction placing the result in the subtrahend. InOutSubtrahend = InMinuend - InOutSubtrahend */
 	SIGNALPROCESSING_API void ArraySubtractInPlace1(TArrayView<const float> InMinuend, TArrayView<float> InOutSubtrahend);
-	SIGNALPROCESSING_API void ArraySubtractInPlace1(const FAlignedFloatBuffer& InMinuend, FAlignedFloatBuffer& InOutSubtrahend);
 
 	/* Performs element-wise in-place subtraction placing the result in the minuend. InOutMinuend = InOutMinuend - InSubtrahend */
 	SIGNALPROCESSING_API void ArraySubtractInPlace2(TArrayView<float> InOutMinuend, TArrayView<const float> InSubtrahend);
-	SIGNALPROCESSING_API void ArraySubtractInPlace2(FAlignedFloatBuffer& InOutMinuend, const FAlignedFloatBuffer& InSubtrahend);
 
 	/** Subtract value from each element in InValues */
 	SIGNALPROCESSING_API void ArraySubtractByConstantInPlace(TArrayView<float> InValues, float InSubtrahend);
-	SIGNALPROCESSING_API void ArraySubtractByConstantInPlace(FAlignedFloatBuffer& InValues, float InSubtrahend);
 
 	/* Square values */
 	SIGNALPROCESSING_API void ArraySquare(TArrayView<const float> InValues, TArrayView<float> OutValues);
@@ -198,62 +181,47 @@ namespace Audio
 
 	/** Perform complex conjugate of array.  Assumes complex numbers are interlaves [real_0, imag_0, real_1, image_1, ..., real_N, imag_N]. */
 	SIGNALPROCESSING_API void ArrayComplexConjugate(TArrayView<const float> InValues, TArrayView<float> OutValues);
-	SIGNALPROCESSING_API void ArrayComplexConjugate(const FAlignedFloatBuffer& InValues, FAlignedFloatBuffer& OutValues);
 
 	SIGNALPROCESSING_API void ArrayComplexConjugateInPlace(TArrayView<float> InValues);
-	SIGNALPROCESSING_API void ArrayComplexConjugateInPlace(FAlignedFloatBuffer& InValues);
 
 	/** Convert magnitude values to decibel values in place. db = 20 * log10(val) */
 	SIGNALPROCESSING_API void ArrayMagnitudeToDecibelInPlace(TArrayView<float> InValues, float InMinimumDb);
-	SIGNALPROCESSING_API void ArrayMagnitudeToDecibelInPlace(FAlignedFloatBuffer& InValues, float InMinimumDb);
 
 	/** Convert power values to decibel values in place. db = 10 * log10(val) */
 	SIGNALPROCESSING_API void ArrayPowerToDecibelInPlace(TArrayView<float> InValues, float InMinimumDb);
-	SIGNALPROCESSING_API void ArrayPowerToDecibelInPlace(FAlignedFloatBuffer& InValues, float InMinimumDb);
 
 	/** Compute power of complex data. Out[i] = Complex[2 * i] * Complex[2 * i] + Complex[2 * i + 1] * Complex[2 * i + 1] */
 	SIGNALPROCESSING_API void ArrayComplexToPower(TArrayView<const float> InComplexSamples, TArrayView<float> OutPowerSamples);
-	SIGNALPROCESSING_API void ArrayComplexToPower(const FAlignedFloatBuffer& InComplexSamples, FAlignedFloatBuffer& OutPowerSamples);
 
 	/** Compute power of complex data. Out[i] = Real[i] * Real[i] + Imaginary[i] * Imaginary[i] */
 	SIGNALPROCESSING_API void ArrayComplexToPower(TArrayView<const float> InRealSamples, TArrayView<const float> InImaginarySamples, TArrayView<float> OutPowerSamples);
-	SIGNALPROCESSING_API void ArrayComplexToPower(const FAlignedFloatBuffer& InRealSamples, const FAlignedFloatBuffer& InImaginarySamples, FAlignedFloatBuffer& OutPowerSamples);
 
 	/* Sets a values to zero if value is denormal. Denormal numbers significantly slow down floating point operations. */
 	SIGNALPROCESSING_API void ArrayUnderflowClamp(TArrayView<float> InOutBuffer);
-	SIGNALPROCESSING_API void ArrayUnderflowClamp(FAlignedFloatBuffer& InOutBuffer);
 
 	/* Clamps the values in a buffer between a min and max value.*/
 	SIGNALPROCESSING_API void ArrayRangeClamp(TArrayView<float> InOutBuffer, float InMinValue, float InMaxValue);
-	SIGNALPROCESSING_API void ArrayRangeClamp(FAlignedFloatBuffer& InOutBuffer, float InMinValue, float InMaxValue);
 
 	/** Sets a constant to a buffer (useful for DC offset application) */
 	SIGNALPROCESSING_API void ArraySetToConstantInplace(TArrayView<float> InOutBuffer, float InConstant);
-	SIGNALPROCESSING_API void ArraySetToConstantInplace(FAlignedFloatBuffer& InOutBuffer, float InConstant);
 
 	/* Performs an element-wise weighted sum OutputBuffer = (InBuffer1 x InGain1) + (InBuffer2 x InGain2) */
 	SIGNALPROCESSING_API void ArrayWeightedSum(TArrayView<const float> InBuffer1, float InGain1, TArrayView<const float> InBuffer2, float InGain2, TArrayView<float> OutBuffer);
-	SIGNALPROCESSING_API void ArrayWeightedSum(const FAlignedFloatBuffer& InBuffer1, float InGain1, const FAlignedFloatBuffer& InBuffer2, float InGain2, FAlignedFloatBuffer& OutBuffer);
 
 	/* Performs an element-wise weighted sum OutputBuffer = (InBuffer1 x InGain1) + InBuffer2 */
 	SIGNALPROCESSING_API void ArrayWeightedSum(TArrayView<const float> InBuffer1, float InGain1, TArrayView<const float> InBuffer2, TArrayView<float> OutBuffer);
-	SIGNALPROCESSING_API void ArrayWeightedSum(const FAlignedFloatBuffer& InBuffer1, float InGain1, const FAlignedFloatBuffer& InBuffer2, FAlignedFloatBuffer& OutBuffer);
 
 	/* Takes a float buffer and quickly interpolates it's gain from StartValue to EndValue. */
 	/* This operation completely ignores channel counts, so avoid using this function on buffers that are not mono, stereo or quad */
 	/* if the buffer needs to fade all channels uniformly. */
 	SIGNALPROCESSING_API void ArrayFade(TArrayView<float> InOutBuffer, const float StartValue, const float EndValue);
-	SIGNALPROCESSING_API void ArrayFade(FAlignedFloatBuffer& OutFloatBuffer, const float StartValue, const float EndValue);
 
 	/** Takes buffer InFloatBuffer, optionally multiplies it by Gain, and adds it to BufferToSumTo. */
 	SIGNALPROCESSING_API void ArrayMixIn(TArrayView<const float> InFloatBuffer, TArrayView<float> BufferToSumTo, const float Gain);
-	SIGNALPROCESSING_API void ArrayMixIn(const FAlignedFloatBuffer& InFloatBuffer, FAlignedFloatBuffer& BufferToSumTo, const float Gain);
 	SIGNALPROCESSING_API void ArrayMixIn(TArrayView<const float> InFloatBuffer, TArrayView<float> BufferToSumTo);
-	SIGNALPROCESSING_API void ArrayMixIn(const FAlignedFloatBuffer& InFloatBuffer, FAlignedFloatBuffer& BufferToSumTo);
 
 	/** This version of ArrayMixIn will fade from StartGain to EndGain. */
 	SIGNALPROCESSING_API void ArrayMixIn(TArrayView<const float> InFloatBuffer, TArrayView<float> BufferToSumTo, const float StartGain, const float EndGain);
-	SIGNALPROCESSING_API void ArrayMixIn(const FAlignedFloatBuffer& InFloatBuffer, FAlignedFloatBuffer& BufferToSumTo, const float StartGain, const float EndGain);
 
 	SIGNALPROCESSING_API void ArrayFloatToPcm16(TArrayView<const float> InView, TArrayView<int16> OutView);
 	SIGNALPROCESSING_API void ArrayPcm16ToFloat(TArrayView<const int16> InView, TArrayView<float> OutView);
diff --git a/Engine/Source/Runtime/SignalProcessing/Public/DSP/MelScale.h b/Engine/Source/Runtime/SignalProcessing/Public/DSP/MelScale.h
--- a/Engine/Source/Runtime/SignalProcessing/Public/DSP/MelScale.h
+++ b/Engine/Source/Runtime/SignalProcessing/Public/DSP/MelScale.h
@@ -1,9 +1,14 @@
 // Copyright Epic Games, Inc. All Rights Reserved.
 
 #pragma once
-#include "CoreMinimal.h"
+
 #include "DSP/FloatArrayMath.h"
+#include "HAL/Platform.h"
+#include "Templates/UniquePtr.h"
 
+#if UE_ENABLE_INCLUDE_ORDER_DEPRECATED_IN_5_1 
+#include "CoreMinimal.h"
+#endif
 
 namespace Audio
 {
