Fix for Limiter Node exploding

diff --git a/Engine/Source/Runtime/SignalProcessing/Private/DynamicsProcessor.cpp b/Engine/Source/Runtime/SignalProcessing/Private/DynamicsProcessor.cpp
--- a/Engine/Source/Runtime/SignalProcessing/Private/DynamicsProcessor.cpp
+++ b/Engine/Source/Runtime/SignalProcessing/Private/DynamicsProcessor.cpp
@@ -255,6 +255,11 @@ namespace Audio
 
 				// Write into the output with the computed gain value
 				OutFrame[Channel] = Gain[Channel] * LookaheadOutput * OutputGain * InputGain;
+
+				if (ProcessingMode == EDynamicsProcessingMode::Limiter)
+				{
+					OutFrame[Channel] = FMath::Clamp(OutFrame[Channel], -ConvertToLinear(ThresholdDb), ConvertToLinear(ThresholdDb));
+				}
 			}
 		}
 	}
