Fix-up copy pasta'ed tooltip in DSP
#rnx

diff --git a/Engine/Source/Runtime/SignalProcessing/Public/DSP/FloatArrayMath.h b/Engine/Source/Runtime/SignalProcessing/Public/DSP/FloatArrayMath.h
--- a/Engine/Source/Runtime/SignalProcessing/Public/DSP/FloatArrayMath.h
+++ b/Engine/Source/Runtime/SignalProcessing/Public/DSP/FloatArrayMath.h
@@ -107,7 +107,9 @@ namespace Audio
 	 */ 
 	SIGNALPROCESSING_API void ArrayMinMaxNormalize(TArrayView<const float> InView, TArray<float>& OutArray);
 
-	// Scale an array to fit within [-1.f, 1.f]
+	/** Returns the largest value of an array irrespective of sign (ex. {-3, 2, 1} would return 3).
+	 *  InView is a view of a float array to get the largest absolute value from.
+	 */
 	SIGNALPROCESSING_API float ArrayMaxAbsValue(const TArrayView<const float> InView);
 
 	/** Multiply the second buffer in place by the first buffer. */
