Adding ability to wait with an event for multithreaded patch to be filled

diff --git a/Engine/Source/Runtime/SignalProcessing/Private/MultithreadedPatching.cpp b/Engine/Source/Runtime/SignalProcessing/Private/MultithreadedPatching.cpp
--- a/Engine/Source/Runtime/SignalProcessing/Private/MultithreadedPatching.cpp
+++ b/Engine/Source/Runtime/SignalProcessing/Private/MultithreadedPatching.cpp
@@ -4,6 +4,8 @@
 
 #include "DSP/FloatArrayMath.h"
 #include "HAL/IConsoleManager.h"
+#include "HAL/PlatformProcess.h"
+#include "HAL/Event.h"
 
 static int32 MultithreadedPatchingPushCallsPerOutputCleanupCheckCVar = 256;
 FAutoConsoleVariableRef CVarMultithreadedPatchingPushCallsPerOutputCleanupCheck(
@@ -23,6 +25,8 @@ namespace Audio
 		, PreviousGain(InGain)
 		, PatchID(++PatchIDCounter)
 		, NumAliveInputs(0)
+		, SamplesFilledEvent(nullptr)
+		, NumSamplesToWaitFor(0)
 	{
 
 	}
@@ -34,9 +38,20 @@ namespace Audio
 		, PreviousGain(0.0f)
 		, PatchID(INDEX_NONE)
 		, NumAliveInputs(0)
+		, SamplesFilledEvent(nullptr)
+		, NumSamplesToWaitFor(0)
 	{
 	}
 
+	FPatchOutput::~FPatchOutput()
+	{
+		if (SamplesFilledEvent)
+		{
+			FPlatformProcess::ReturnSynchEventToPool(SamplesFilledEvent);
+			SamplesFilledEvent = nullptr;
+		}
+	}
+
 	int32 FPatchOutput::PopAudio(float* OutBuffer, int32 NumSamples, bool bUseLatestAudio)
 	{
 		if (IsInputStale())
@@ -44,6 +59,10 @@ namespace Audio
 			return INDEX_NONE;
 		}
 
+		const int32 CurrSamplesAvailable = GetNumSamplesAvailable();
+		const int32 CurrCapacity = InternalBuffer.GetCapacity();
+		const int32 CurrSamplesFilled = CurrCapacity - CurrSamplesAvailable;
+		
 		if (bUseLatestAudio && InternalBuffer.Num() > (uint32) NumSamples)
 		{
 			InternalBuffer.SetNum((uint32)NumSamples);
@@ -72,6 +91,25 @@ namespace Audio
 		return NumAliveInputs == 0;
 	}
 
+	int32 FPatchOutput::PushAudioToInternalBuffer(const float* InBuffer, int32 NumSamples)
+	{
+		const int32 NumSamplesPushed = InternalBuffer.Push(InBuffer, NumSamples);
+
+		// Check to see if we need to notify anybody waiting on this patch output getting filled
+		if (SamplesFilledEvent && NumSamplesToWaitFor > 0 && GetNumSamplesAvailable() >= NumSamplesToWaitFor)
+		{
+			SamplesFilledEvent->Trigger();
+			NumSamplesToWaitFor = 0;
+		}
+
+		const int32 CurrSamplesAvailable = GetNumSamplesAvailable();
+		const int32 CurrCapacity = InternalBuffer.GetCapacity();
+		const int32 CurrSamplesFilled = CurrCapacity - CurrSamplesAvailable;
+		
+		return NumSamplesPushed;
+	}
+
+
 	int32 FPatchOutput::MixInAudio(float* OutBuffer, int32 NumSamples, bool bUseLatestAudio)
 	{
 		if (IsInputStale())
@@ -110,11 +148,28 @@ namespace Audio
 	}
 
 
-	int32 FPatchOutput::GetNumSamplesAvailable()
+	int32 FPatchOutput::GetNumSamplesAvailable() const
 	{
 		return InternalBuffer.Num();
 	}
 
+	bool FPatchOutput::WaitUntilNumSamplesAvailable(int32 InNumSamplesToWaitFor, uint32 TimeOutMilliseconds)
+	{
+		// No need to wait if we have that many available already
+		if (GetNumSamplesAvailable() >= InNumSamplesToWaitFor)
+		{
+			return true;
+		}
+		NumSamplesToWaitFor = InNumSamplesToWaitFor;
+		
+		if (!SamplesFilledEvent)
+		{
+			SamplesFilledEvent = FPlatformProcess::GetSynchEventFromPool(false);
+		}
+	
+		return SamplesFilledEvent->Wait(TimeOutMilliseconds);
+	}
+
 	FPatchInput::FPatchInput(const FPatchOutputStrongPtr& InOutput)
 		: OutputHandle(InOutput)
 	{
@@ -178,7 +233,7 @@ namespace Audio
 			return INDEX_NONE;
 		}
 
-		int32 SamplesPushed = OutputHandle->InternalBuffer.Push(InBuffer, NumSamples);
+		int32 SamplesPushed = OutputHandle->PushAudioToInternalBuffer(InBuffer, NumSamples);
 
 		// Periodically check to see if the output handle has been destroyed and clean it up.
 		// If the buffer is full, check as well to determine if it is possibly due to the
diff --git a/Engine/Source/Runtime/SignalProcessing/Public/DSP/MultithreadedPatching.h b/Engine/Source/Runtime/SignalProcessing/Public/DSP/MultithreadedPatching.h
--- a/Engine/Source/Runtime/SignalProcessing/Public/DSP/MultithreadedPatching.h
+++ b/Engine/Source/Runtime/SignalProcessing/Public/DSP/MultithreadedPatching.h
@@ -6,8 +6,12 @@
 #include "DSP/BufferVectorOperations.h"
 #include "DSP/Dsp.h"
 
+class FEvent;
+
 namespace Audio
 {
+
+	
 	/**
 	 * This class can be thought of as an output for a single constructed instance of FPatchInput.
 	 * Each FPatchOutput can only be connected to one FPatchInput. To route multiple outputs, see FPatchSplitter.
@@ -33,6 +37,8 @@ namespace Audio
 		/** The default constructor will result in an uninitialized, disconnected patch point. */
 		FPatchOutput();
 
+		virtual ~FPatchOutput();
+		
 		/** Copies the minimum of NumSamples or however many samples are available into OutBuffer. Returns the number of samples copied, or -1 if this output's corresponding input has been destroyed. */
 		int32 PopAudio(float* OutBuffer, int32 NumSamples, bool bUseLatestAudio);
 
@@ -40,8 +46,11 @@ namespace Audio
 		int32 MixInAudio(float* OutBuffer, int32 NumSamples, bool bUseLatestAudio);
 
 		/** Returns the current number of samples buffered on this output. */
-		int32 GetNumSamplesAvailable();
+		int32 GetNumSamplesAvailable() const;
 
+		/** Pauses the current thread until there are the given number of samples available to pop. Will return true if it succeeded, false if it timed out. */
+		bool WaitUntilNumSamplesAvailable(int32 NumSamples, uint32 TimeOutMilliseconds = MAX_uint32);
+		
 		/** Returns true if the input for this patch has been destroyed. */
 		bool IsInputStale() const;
 
@@ -49,7 +58,9 @@ namespace Audio
 		friend class FPatchMixer;
 		friend class FPatchSplitter;
 	private:
-
+		
+		int32 PushAudioToInternalBuffer(const float* InBuffer, int32 NumSamples);
+		
 		// Internal buffer.
 		TCircularAudioBuffer<float> InternalBuffer;
 
@@ -66,6 +77,10 @@ namespace Audio
 		// Counter that is incremented/decremented to allow FPatchInput to be copied around safely.
 		int32 NumAliveInputs;
 
+		// Event to pause the current thread until a given number of samples has been filled
+		FEvent* SamplesFilledEvent;
+		int32 NumSamplesToWaitFor;
+		
 		static TAtomic<int32> PatchIDCounter;
 	};
 
