Inclusive language push, removing master/slave language from DSP code

diff --git a/Engine/Plugins/Runtime/Synthesis/Source/Synthesis/Classes/SynthComponents/EpicSynth1Component.h b/Engine/Plugins/Runtime/Synthesis/Source/Synthesis/Classes/SynthComponents/EpicSynth1Component.h
--- a/Engine/Plugins/Runtime/Synthesis/Source/Synthesis/Classes/SynthComponents/EpicSynth1Component.h
+++ b/Engine/Plugins/Runtime/Synthesis/Source/Synthesis/Classes/SynthComponents/EpicSynth1Component.h
@@ -469,7 +469,7 @@ public:
 	UFUNCTION(BlueprintCallable, Category = "Synth|Components|Audio")
 	void SetEnableUnison(bool EnableUnison);
 
-	// Set whether or not to slave the phase of osc2 to osc1
+	// Set whether or not to follow the phase of osc2 to osc1
 	UFUNCTION(BlueprintCallable, Category = "Synth|Components|Audio")
 	void SetOscSync(const bool bIsSynced);
 
diff --git a/Engine/Plugins/Runtime/Synthesis/Source/Synthesis/Private/EpicSynth1.cpp b/Engine/Plugins/Runtime/Synthesis/Source/Synthesis/Private/EpicSynth1.cpp
--- a/Engine/Plugins/Runtime/Synthesis/Source/Synthesis/Private/EpicSynth1.cpp
+++ b/Engine/Plugins/Runtime/Synthesis/Source/Synthesis/Private/EpicSynth1.cpp
@@ -59,8 +59,8 @@ namespace Audio
 			OscilPan[i].Init(VoiceId, ModMatrix);
 		}
 
-		// Setup Osc1 to be a master to Osc2 
-		Oscil[0].SetSlaveOsc(&Oscil[1]);
+		// Setup Osc1 to be a leader to Osc2 
+		Oscil[0].SetFollowerOsc(&Oscil[1]);
 
 		Amp.Init(VoiceId, ModMatrix);
 
diff --git a/Engine/Source/Runtime/SignalProcessing/Private/Osc.cpp b/Engine/Source/Runtime/SignalProcessing/Private/Osc.cpp
--- a/Engine/Source/Runtime/SignalProcessing/Private/Osc.cpp
+++ b/Engine/Source/Runtime/SignalProcessing/Private/Osc.cpp
@@ -19,7 +19,7 @@ namespace Audio
 		, PulseWidthMod(0.0f)
 		, PulseWidth(0.0f)
 		, ModMatrix(nullptr)
-		, SlaveOsc(nullptr)
+		, FollowerOsc(nullptr)
 		, bIsPlaying(false)
 		, bChanged(false)
 	{
@@ -182,9 +182,9 @@ namespace Audio
 		Phase = 0.0f;
 	}
 
-	void IOscBase::SetSlaveOsc(IOscBase* InSlaveOsc)
+	void IOscBase::SetFollowerOsc(IOscBase* InFollowerOsc)
 	{
-		SlaveOsc = InSlaveOsc;
+		FollowerOsc = InFollowerOsc;
 	}
 
 	void IOscBase::Reset()
diff --git a/Engine/Source/Runtime/SignalProcessing/Public/DSP/Osc.h b/Engine/Source/Runtime/SignalProcessing/Public/DSP/Osc.h
--- a/Engine/Source/Runtime/SignalProcessing/Public/DSP/Osc.h
+++ b/Engine/Source/Runtime/SignalProcessing/Public/DSP/Osc.h
@@ -122,14 +122,14 @@ namespace Audio
 		// Returns whether or not this oscillator is playing
 		bool IsPlaying() const { return bIsPlaying; }
 
-		// Returns if this oscillator should be synced to a master oscillator
+		// Returns if this oscillator should be synced to a leader oscillator
 		bool IsSync() const { return bIsSync; }
 
-		// Sets whether or not this oscillator should be synced to a master oscillator. Master oscillator needs to have set this oscillator as its slave.
+		// Sets whether or not this oscillator should be synced to a leader oscillator. leader oscillator needs to have set this oscillator as its follower.
 		void SetSync(const bool bInSync) { bIsSync = bInSync; }
 
-		// Sets the input oscillator as the slave of this oscillator
-		void SetSlaveOsc(IOscBase* InSlaveOsc);
+		// Sets the input oscillator as the follower of this oscillator
+		void SetFollowerOsc(IOscBase* InFollowerOsc);
 
 		// Return patch destinations for various modulatable parameters
 		FPatchDestination GetModDestFrequency() const { return ModFrequencyDest; }
@@ -161,9 +161,9 @@ namespace Audio
 				Result = true;
 			}
 
-			if (Result && SlaveOsc && SlaveOsc->IsSync())
+			if (Result && FollowerOsc && FollowerOsc->IsSync())
 			{
-				SlaveOsc->ResetPhase();
+				FollowerOsc->ResetPhase();
 			}
 
 			return Result;
@@ -231,8 +231,8 @@ namespace Audio
 		FPatchDestination ModScaleDest;
 		FPatchDestination ModAddDest;
 
-		// Ptr to a slave oscillator that can be triggered to 0 phase if it is synced.
-		IOscBase* SlaveOsc;
+		// Ptr to a follower oscillator that can be triggered to 0 phase if it is synced.
+		IOscBase* FollowerOsc;
 
 		// Whether or not the oscillator is on or off
 		bool bIsPlaying;
