Add an Audio_DSP csv stat category, disabled by default

diff --git a/Engine/Source/Runtime/SignalProcessing/Private/FloatArrayMath.cpp b/Engine/Source/Runtime/SignalProcessing/Private/FloatArrayMath.cpp
--- a/Engine/Source/Runtime/SignalProcessing/Private/FloatArrayMath.cpp
+++ b/Engine/Source/Runtime/SignalProcessing/Private/FloatArrayMath.cpp
@@ -4,6 +4,7 @@
 #include "CoreMinimal.h"
 #include "Math/UnrealMathVectorConstants.h"
 #include "SignalProcessingModule.h"
+#include "ProfilingDebugging/CsvProfiler.h"
 
 #if INTEL_ISPC && !UE_BUILD_SHIPPING
 #include "HAL/IConsoleManager.h"
@@ -18,6 +19,8 @@ bool bAudio_FloatArrayMath_ISPC_Enabled = true;
 FAutoConsoleVariableRef CVarAudioFloatArrayMathISPCEnabled(TEXT("au.FloatArrayMath.ISPC"), bAudio_FloatArrayMath_ISPC_Enabled, TEXT("Whether to use ISPC optimizations in audio float array math operations"));
 #endif
 
+CSV_DEFINE_CATEGORY(Audio_Dsp, false);
+
 namespace Audio
 {
 	namespace MathIntrinsics
@@ -29,6 +32,7 @@ namespace Audio
 
 	void ArraySum(TArrayView<const float> InValues, float& OutSum)
 	{
+		CSV_SCOPED_TIMING_STAT(Audio_Dsp, ArraySum);
 		OutSum = 0.f;
 
 		int32 Num = InValues.Num();
@@ -72,6 +76,7 @@ namespace Audio
 
 	void ArraySum(TArrayView<const float> InFloatBuffer1, TArrayView<const float> InFloatBuffer2, TArrayView<float> OutputBuffer)
 	{
+		CSV_SCOPED_TIMING_STAT(Audio_Dsp, ArraySum);
 		checkf(InFloatBuffer1.Num() == InFloatBuffer2.Num(), TEXT("Input buffers must be equal length"));
 
 		const int32 Num = InFloatBuffer1.Num();
@@ -120,6 +125,7 @@ namespace Audio
 		{
 			return;
 		}
+		CSV_SCOPED_TIMING_STAT(Audio_Dsp, ArrayCumulativeSum);
 
 		float* OutDataPtr = OutData.GetData();
 		const float* InViewPtr = InView.GetData();
@@ -153,6 +159,8 @@ namespace Audio
 		{
 			return;
 		}
+		CSV_SCOPED_TIMING_STAT(Audio_Dsp, ArrayMean);
+
 
 		const float* DataPtr = InView.GetData();
 
@@ -183,6 +191,7 @@ namespace Audio
 		{
 			return;
 		}
+		CSV_SCOPED_TIMING_STAT(Audio_Dsp, ArrayMeanSquared);
 
 		const float* DataPtr = InView.GetData();
 
@@ -205,6 +214,7 @@ namespace Audio
 
 	float ArrayGetMagnitude(TArrayView<const float> Buffer)
 	{
+		CSV_SCOPED_TIMING_STAT(Audio_Dsp, ArrayGetMagnitude);
 		const int32 Num = Buffer.Num();
 
 		if (bAudio_FloatArrayMath_ISPC_Enabled)
@@ -253,6 +263,7 @@ namespace Audio
 
 	float ArrayGetAverageValue(TArrayView<const float> Buffer)
 	{
+		CSV_SCOPED_TIMING_STAT(Audio_Dsp, ArrayGetAverageValue);
 		const int32 Num = Buffer.Num();
 
 		if (bAudio_FloatArrayMath_ISPC_Enabled)
@@ -298,6 +309,7 @@ namespace Audio
 
 	float ArrayGetAverageAbsValue(TArrayView<const float> Buffer)
 	{
+		CSV_SCOPED_TIMING_STAT(Audio_Dsp, ArrayGetAverageAbsValue);
 		const int32 Num = Buffer.Num();
 
 		if (bAudio_FloatArrayMath_ISPC_Enabled)
@@ -357,7 +369,8 @@ namespace Audio
 		{
 			return;
 		}
-		
+		CSV_SCOPED_TIMING_STAT(Audio_Dsp, ArrayMeanFilter);
+
 		// Use cumulative sum to avoid multiple summations 
 		// Instead of summing over InView[StartIndex:EndIndex], avoid all that
 		// calculation by taking difference of cumulative sum at those two points:
@@ -434,6 +447,7 @@ namespace Audio
 		{
 			return;
 		}
+		CSV_SCOPED_TIMING_STAT(Audio_Dsp, ArrayMaxFilter);
 
 		// Get max in first window
 		int32 ActualStartIndex = 0;
@@ -496,6 +510,8 @@ namespace Audio
 
 	void ArrayGetEuclideanNorm(TArrayView<const float> InView, float& OutEuclideanNorm)
 	{
+		CSV_SCOPED_TIMING_STAT(Audio_Dsp, ArrayGetEuclideanNorm);
+
 		// Initialize output.
 		OutEuclideanNorm = 0.0f;
 		const int32 Num = InView.Num();
@@ -521,6 +537,7 @@ namespace Audio
 
 	void ArrayAbs(TArrayView<const float> InBuffer, TArrayView<float> OutBuffer)
 	{
+		CSV_SCOPED_TIMING_STAT(Audio_Dsp, ArrayAbs);
 		const int32 Num = InBuffer.Num();
 		check(OutBuffer.Num() == Num);
 
@@ -545,6 +562,7 @@ namespace Audio
 
 	void ArrayAbsInPlace(TArrayView<float> InView)
 	{
+		CSV_SCOPED_TIMING_STAT(Audio_Dsp, ArrayAbsInPlace);
 		const int32 Num = InView.Num();
 		float* Data = InView.GetData();
 
@@ -565,6 +583,7 @@ namespace Audio
 
 	void ArrayClampMinInPlace(TArrayView<float> InView, float InMin)
 	{
+		CSV_SCOPED_TIMING_STAT(Audio_Dsp, ArrayClampMinInPlace);
 		const int32 Num = InView.Num();
 		float* Data = InView.GetData();
 
@@ -585,6 +604,7 @@ namespace Audio
 
 	void ArrayClampMaxInPlace(TArrayView<float> InView, float InMax)
 	{
+		CSV_SCOPED_TIMING_STAT(Audio_Dsp, ArrayClampMaxInPlace);
 		const int32 Num = InView.Num();
 		float* Data = InView.GetData();
 
@@ -605,6 +625,8 @@ namespace Audio
 
 	void ArrayClampInPlace(TArrayView<float> InView, float InMin, float InMax)
 	{
+		CSV_SCOPED_TIMING_STAT(Audio_Dsp, ArrayClampInPlace);
+
 		const int32 Num = InView.Num();
 		float* Data = InView.GetData();
 
@@ -632,6 +654,7 @@ namespace Audio
 		{
 			return;
 		}
+		CSV_SCOPED_TIMING_STAT(Audio_Dsp, ArrayMinMaxNormalize);
 
 		OutArray.AddUninitialized(Num);
 
@@ -673,6 +696,8 @@ namespace Audio
 	
 	float ArrayMaxAbsValue(const TArrayView<const float> InView)
 	{
+		CSV_SCOPED_TIMING_STAT(Audio_Dsp, ArrayMaxAbsValue);
+
 		const int32 Num = InView.Num();
 		const float* Data = InView.GetData();
 
@@ -709,6 +734,7 @@ namespace Audio
 
 	void ArrayMultiplyInPlace(TArrayView<const float> InFloatBuffer, TArrayView<float> BufferToMultiply)
 	{
+		CSV_SCOPED_TIMING_STAT(Audio_Dsp, ArrayMultiplyInPlace);
 		checkf(InFloatBuffer.Num() == BufferToMultiply.Num(), TEXT("Input buffers must be equal length"));
 
 		const int32 Num = BufferToMultiply.Num();
@@ -748,6 +774,7 @@ namespace Audio
 
 	void ArrayComplexMultiplyInPlace(TArrayView<const float> InValues1, TArrayView<float> InValues2)
 	{
+		CSV_SCOPED_TIMING_STAT(Audio_Dsp, ArrayComplexMultiplyInPlace);
 		check(InValues1.Num() == InValues2.Num());
 
 		const int32 Num = InValues1.Num();
@@ -805,6 +832,7 @@ namespace Audio
 
 	void ArrayMultiplyByConstant(TArrayView<const float> InFloatBuffer, float InValue, TArrayView<float> OutFloatBuffer)
 	{
+		CSV_SCOPED_TIMING_STAT(Audio_Dsp, ArrayMultiplyByConstant);
 		check(InFloatBuffer.Num() == OutFloatBuffer.Num());
 
 		const int32 Num = InFloatBuffer.Num();
@@ -855,6 +883,7 @@ namespace Audio
 
 	void ArrayMultiplyByConstantInPlace(TArrayView<float> InOutBuffer, float InGain)
 	{
+		CSV_SCOPED_TIMING_STAT(Audio_Dsp, ArrayMultiplyByConstantInPlace);
 		int32 Num = InOutBuffer.Num();
 		float* InOutData = InOutBuffer.GetData();
 
@@ -893,6 +922,8 @@ namespace Audio
 
 	void ArrayAddInPlace(TArrayView<const float> InValues, TArrayView<float> InAccumulateValues)
 	{
+		CSV_SCOPED_TIMING_STAT(Audio_Dsp, ArrayAddInPlace);
+
 		check(InValues.Num() == InAccumulateValues.Num());
 
 		const int32 Num = InValues.Num();
@@ -935,6 +966,7 @@ namespace Audio
 
 	void ArrayAddConstantInplace(TArrayView<float> InOutBuffer, float InConstant)
 	{
+		CSV_SCOPED_TIMING_STAT(Audio_Dsp, ArrayAddConstantInplace);
 		int32 Num = InOutBuffer.Num();
 		float* InOutData = InOutBuffer.GetData();
 
@@ -973,6 +1005,7 @@ namespace Audio
 
 	void ArrayMultiplyAddInPlace(TArrayView<const float> InValues, float InMultiplier, TArrayView<float> InAccumulateValues)
 	{
+		CSV_SCOPED_TIMING_STAT(Audio_Dsp, ArrayMultiplyAddInPlace);
 		check(InValues.Num() == InAccumulateValues.Num());
 		
 		const int32 Num = InValues.Num();
@@ -997,6 +1030,7 @@ namespace Audio
 
 	void ArrayLerpAddInPlace(TArrayView<const float> InValues, float InStartMultiplier, float InEndMultiplier, TArrayView<float> InAccumulateValues)
 	{
+		CSV_SCOPED_TIMING_STAT(Audio_Dsp, ArrayLerpAddInPlace);
 		check(InValues.Num() == InAccumulateValues.Num());
 
 		const int32 Num = InValues.Num();
@@ -1051,6 +1085,7 @@ namespace Audio
 	/* Subtracts two buffers together element-wise. */
 	void ArraySubtract(TArrayView<const float> InMinuend, TArrayView<const float> InSubtrahend, TArrayView<float> OutBuffer)
 	{
+		CSV_SCOPED_TIMING_STAT(Audio_Dsp, ArraySubtract);
 		const int32 Num = InMinuend.Num();
 
 		checkf(Num == InSubtrahend.Num() && Num == OutBuffer.Num(), TEXT("InMinuend, InSubtrahend, and OutBuffer must have equal Num elements (%d vs %d vs %d)"), Num, InSubtrahend.Num(), OutBuffer.Num());
@@ -1090,6 +1125,7 @@ namespace Audio
 	/* Performs element-wise in-place subtraction placing the result in the subtrahend. InOutSubtrahend = InMinuend - InOutSubtrahend */
 	void ArraySubtractInPlace1(TArrayView<const float> InMinuend, TArrayView<float> InOutSubtrahend)
 	{
+		CSV_SCOPED_TIMING_STAT(Audio_Dsp, ArraySubtractInPlace1);
 		checkf(InMinuend.Num() == InOutSubtrahend.Num(), TEXT("Input buffers must be equal length"));
 
 		const int32 Num = InMinuend.Num();
@@ -1130,6 +1166,7 @@ namespace Audio
 	/* Performs element-wise in-place subtraction placing the result in the minuend. InOutMinuend = InOutMinuend - InSubtrahend */
 	void ArraySubtractInPlace2(TArrayView<float> InOutMinuend, TArrayView<const float> InSubtrahend)
 	{
+		CSV_SCOPED_TIMING_STAT(Audio_Dsp, ArraySubtractInPlace2);
 		checkf(InOutMinuend.Num() == InSubtrahend.Num(), TEXT("Input buffers must be equal length"));
 
 		const int32 Num = InOutMinuend.Num();
@@ -1169,6 +1206,7 @@ namespace Audio
 
 	void ArraySubtractByConstantInPlace(TArrayView<float> InValues, float InSubtrahend)
 	{
+		CSV_SCOPED_TIMING_STAT(Audio_Dsp, ArraySubtractByConstantInPlace);
 		const int32 Num = InValues.Num();
 		float* InData = InValues.GetData();
 		if (bAudio_FloatArrayMath_ISPC_Enabled)
@@ -1206,6 +1244,7 @@ namespace Audio
 
 	void ArraySquare(TArrayView<const float> InValues, TArrayView<float> OutValues)
 	{
+		CSV_SCOPED_TIMING_STAT(Audio_Dsp, ArraySquare);
 		check(InValues.Num() == OutValues.Num());
 
 		const int32 Num = InValues.Num();
@@ -1246,6 +1285,7 @@ namespace Audio
 
 	void ArraySquareInPlace(TArrayView<float> InValues)
 	{
+		CSV_SCOPED_TIMING_STAT(Audio_Dsp, ArraySquareInPlace);
 		const int32 Num = InValues.Num();
 
 		float* InData = InValues.GetData();
@@ -1283,6 +1323,8 @@ namespace Audio
 
 	void ArraySqrtInPlace(TArrayView<float> InValues)
 	{
+		CSV_SCOPED_TIMING_STAT(Audio_Dsp, ArraySqrtInPlace);
+
 		const int32 Num = InValues.Num();
 		float* InValuesData = InValues.GetData();
 
@@ -1303,6 +1345,8 @@ namespace Audio
 
 	void ArrayComplexConjugate(TArrayView<const float> InValues, TArrayView<float> OutValues)
 	{
+		CSV_SCOPED_TIMING_STAT(Audio_Dsp, ArrayComplexConjugate);
+
 		check(OutValues.Num() == InValues.Num());
 		check((InValues.Num() % 2) == 0);
 
@@ -1349,6 +1393,8 @@ namespace Audio
 
 	void ArrayComplexConjugateInPlace(TArrayView<float> InValues)
 	{
+		CSV_SCOPED_TIMING_STAT(Audio_Dsp, ArrayComplexConjugateInPlace);
+
 		check((InValues.Num() % 2) == 0);
 
 		int32 Num = InValues.Num();
@@ -1392,6 +1438,8 @@ namespace Audio
 
 	void ArrayMagnitudeToDecibelInPlace(TArrayView<float> InValues, float InMinimumDb)
 	{
+		CSV_SCOPED_TIMING_STAT(Audio_Dsp, ArrayMagnitudeToDecibelInPlace);
+
 		const int32 Num = InValues.Num();
 		float* InValuesData = InValues.GetData();
 
@@ -1439,6 +1487,8 @@ namespace Audio
 
 	void ArrayPowerToDecibelInPlace(TArrayView<float> InValues, float InMinimumDb)
 	{
+		CSV_SCOPED_TIMING_STAT(Audio_Dsp, ArrayPowerToDecibelInPlace);
+
 		const int32 Num = InValues.Num();
 		float* InValuesData = InValues.GetData();
 
@@ -1486,6 +1536,8 @@ namespace Audio
 
 	void ArrayComplexToPower(TArrayView<const float> InComplexValues, TArrayView<float> OutPowerValues)
 	{
+		CSV_SCOPED_TIMING_STAT(Audio_Dsp, ArrayComplexToPower);
+
 		check((InComplexValues.Num() % 2) == 0);
 		check(InComplexValues.Num() == (OutPowerValues.Num() * 2));
 
@@ -1541,6 +1593,8 @@ namespace Audio
 
 	void ArrayComplexToPower(TArrayView<const float> InRealSamples, TArrayView<const float> InImaginarySamples, TArrayView<float> OutPowerSamples)
 	{
+		CSV_SCOPED_TIMING_STAT(Audio_Dsp, ArrayComplexToPower);
+
 		checkf(InRealSamples.Num() == InImaginarySamples.Num(), TEXT("Input buffers must have equal number of elements"));
 
 		const int32 Num = InRealSamples.Num();
@@ -1586,6 +1640,8 @@ namespace Audio
 	/* Sets a values to zero if value is denormal. Denormal numbers significantly slow down floating point operations. */
 	void ArrayUnderflowClamp(TArrayView<float> InOutValues)
 	{
+		CSV_SCOPED_TIMING_STAT(Audio_Dsp, ArrayUnderflowClamp);
+
 		int32 Num = InOutValues.Num();
 		float* InOutData = InOutValues.GetData();
 
@@ -1638,6 +1694,8 @@ namespace Audio
 	/* Clamps values in the buffer to be between InMinValue and InMaxValue */
 	void ArrayRangeClamp(TArrayView<float> InOutBuffer, float InMinValue, float InMaxValue)
 	{
+		CSV_SCOPED_TIMING_STAT(Audio_Dsp, ArrayRangeClamp);
+
 		int32 Num = InOutBuffer.Num();
 		float* InOutData = InOutBuffer.GetData();
 
@@ -1685,6 +1743,8 @@ namespace Audio
 
 	void ArraySetToConstantInplace(TArrayView<float> InOutBuffer, float InConstant)
 	{
+		CSV_SCOPED_TIMING_STAT(Audio_Dsp, ArraySetToConstantInplace);
+
 		int32 Num = InOutBuffer.Num();
 		float* InOutData = InOutBuffer.GetData();
 
@@ -1722,6 +1782,8 @@ namespace Audio
 	/* Performs an element-wise weighted sum OutputBuffer = (InBuffer1 x InGain1) + (InBuffer2 x InGain2) */
 	void ArrayWeightedSum(TArrayView<const float> InBuffer1, float InGain1, TArrayView<const float> InBuffer2, float InGain2, TArrayView<float> OutBuffer)
 	{
+		CSV_SCOPED_TIMING_STAT(Audio_Dsp, ArrayWeightedSum);
+
 		checkf(InBuffer1.Num() == InBuffer2.Num(), TEXT("Buffers must be equal length"));
 
 		int32 Num = InBuffer1.Num();
@@ -1772,6 +1834,7 @@ namespace Audio
 	/* Performs an element-wise weighted sum OutputBuffer = (InBuffer1 x InGain1) + InBuffer2 */
 	void ArrayWeightedSum(TArrayView<const float> InBuffer1, float InGain1, TArrayView<const float> InBuffer2, TArrayView<float> OutBuffer)
 	{
+		CSV_SCOPED_TIMING_STAT(Audio_Dsp, ArrayWeightedSum);
 		checkf(InBuffer1.Num() == InBuffer2.Num() && InBuffer1.Num() == OutBuffer.Num(), TEXT("Buffers must be equal length"));
 
 		int32 Num = InBuffer1.Num();
@@ -1816,6 +1879,8 @@ namespace Audio
 
 	void ArrayFade(TArrayView<float> InOutBuffer, const float StartValue, const float EndValue)
 	{
+		CSV_SCOPED_TIMING_STAT(Audio_Dsp, ArrayFade);
+
 		int32 Num = InOutBuffer.Num();
 		float* OutFloatBuffer = InOutBuffer.GetData();
 
@@ -1879,6 +1944,8 @@ namespace Audio
 
 	void ArrayMixIn(TArrayView<const float> InFloatBuffer, TArrayView<float> BufferToSumTo, const float Gain)
 	{
+		CSV_SCOPED_TIMING_STAT(Audio_Dsp, ArrayMixIn);
+
 		checkf(InFloatBuffer.Num() == BufferToSumTo.Num(), TEXT("Buffers must be equal size"));
 
 		int32 Num = InFloatBuffer.Num();
@@ -1921,6 +1988,8 @@ namespace Audio
 
 	void ArrayMixIn(TArrayView<const float> InFloatBuffer, TArrayView<float> BufferToSumTo)
 	{
+		CSV_SCOPED_TIMING_STAT(Audio_Dsp, ArrayMixIn);
+
 		checkf(InFloatBuffer.Num() == BufferToSumTo.Num(), TEXT("Buffers must be equal size"));
 
 		int32 Num = InFloatBuffer.Num();
@@ -1961,6 +2030,8 @@ namespace Audio
 
 	void ArrayMixIn(TArrayView<const float> InFloatBuffer, TArrayView<float> BufferToSumTo, const float StartGain, const float EndGain)
 	{
+		CSV_SCOPED_TIMING_STAT(Audio_Dsp, ArrayMixIn);
+
 		checkf(InFloatBuffer.Num() == BufferToSumTo.Num(), TEXT("Buffers must be equal size"));
 
 		int32 Num = InFloatBuffer.Num();
@@ -2027,6 +2098,8 @@ namespace Audio
 
 	void ArrayFloatToPcm16(TArrayView<const float> InView, TArrayView<int16> OutView)
 	{
+		CSV_SCOPED_TIMING_STAT(Audio_Dsp, ArrayFloatToPcm16);
+
 		check(InView.Num() == OutView.Num());
 
 		const int32 Num = InView.Num();
@@ -2069,6 +2142,8 @@ namespace Audio
 	
 	void ArrayPcm16ToFloat(TArrayView<const int16> InView, TArrayView<float> OutView)
 	{
+		CSV_SCOPED_TIMING_STAT(Audio_Dsp, ArrayPcm16ToFloat);
+
 		check(InView.Num() == OutView.Num());
 
 		const int32 Num = InView.Num();
@@ -2115,7 +2190,6 @@ namespace Audio
 		{
 			return;
 		}
-		
 		const int32 NumChannels = InBuffers.Num();
 		const int32 NumFrames = InBuffers[0].Num();
 
@@ -2137,6 +2211,7 @@ namespace Audio
 
 	void ArrayInterleave(const float** RESTRICT InBuffers, float* RESTRICT OutBuffer, const int32 InFrames, const int32 InChannels)
 	{
+		CSV_SCOPED_TIMING_STAT(Audio_Dsp, ArrayInterleave);
 		for(int32 ChannelIdx = 0; ChannelIdx < InChannels; ChannelIdx++)
 		{
 			const float* InBuffer = InBuffers[ChannelIdx];
@@ -2174,6 +2249,8 @@ namespace Audio
 
 	void ArrayDeinterleave(const float* RESTRICT InBuffer, float** RESTRICT OutBuffers, const int32 InFrames, const int32 InChannels)
 	{
+		CSV_SCOPED_TIMING_STAT(Audio_Dsp, ArrayDeinterleave);
+
 		for(int32 ChannelIdx = 0; ChannelIdx < InChannels; ChannelIdx++)
 		{
 			float* OutBuffer = OutBuffers[ChannelIdx];
@@ -2252,6 +2329,8 @@ namespace Audio
 
 	void FContiguousSparse2DKernelTransform::TransformArray(const float* InArray, float* OutArray) const
 	{
+		CSV_SCOPED_TIMING_STAT(Audio_Dsp, TransformArray);
+
 		check(nullptr != InArray);
 		check(nullptr != OutArray);
 
