Make it possible to disable AudioFloatArrayMath ISPC per platform, and opt out on platforms where it doesn't help performance

diff --git a/Engine/Source/Runtime/SignalProcessing/Private/FloatArrayMath.cpp b/Engine/Source/Runtime/SignalProcessing/Private/FloatArrayMath.cpp
--- a/Engine/Source/Runtime/SignalProcessing/Private/FloatArrayMath.cpp
+++ b/Engine/Source/Runtime/SignalProcessing/Private/FloatArrayMath.cpp
@@ -14,9 +14,16 @@
 #include "FloatArrayMath.ispc.generated.h"
 #endif
 
-#if INTEL_ISPC && !UE_BUILD_SHIPPING
-bool bAudio_FloatArrayMath_ISPC_Enabled = true;
-FAutoConsoleVariableRef CVarAudioFloatArrayMathISPCEnabled(TEXT("au.FloatArrayMath.ISPC"), bAudio_FloatArrayMath_ISPC_Enabled, TEXT("Whether to use ISPC optimizations in audio float array math operations"));
+#if !defined(AUDIO_FLOAT_ARRAY_MATH_ISPC_ENABLED_DEFAULT)
+#define AUDIO_FLOAT_ARRAY_MATH_ISPC_ENABLED_DEFAULT 1
+#endif
+
+// Support run-time toggling on supported platforms in non-shipping configurations
+#if !INTEL_ISPC || UE_BUILD_SHIPPING
+static constexpr bool bAudio_FloatArrayMath_ISPC_Enabled = INTEL_ISPC && AUDIO_FLOAT_ARRAY_MATH_ISPC_ENABLED_DEFAULT;
+#else
+static bool bAudio_FloatArrayMath_ISPC_Enabled = AUDIO_FLOAT_ARRAY_MATH_ISPC_ENABLED_DEFAULT;
+static FAutoConsoleVariableRef CVarAudioFloatArrayMathISPCEnabled(TEXT("au.FloatArrayMath.ISPC"), bAudio_FloatArrayMath_ISPC_Enabled, TEXT("Whether to use ISPC optimizations in audio float array math operations"));
 #endif
 
 CSV_DEFINE_CATEGORY(Audio_Dsp, false);
diff --git a/Engine/Source/Runtime/SignalProcessing/Public/DSP/FloatArrayMath.h b/Engine/Source/Runtime/SignalProcessing/Public/DSP/FloatArrayMath.h
--- a/Engine/Source/Runtime/SignalProcessing/Public/DSP/FloatArrayMath.h
+++ b/Engine/Source/Runtime/SignalProcessing/Public/DSP/FloatArrayMath.h
@@ -314,12 +314,3 @@ namespace Audio
 	};
 
 }
-
-// Support ISPC enable/disable in non-shipping builds
-#if !INTEL_ISPC
-const bool bAudio_FloatArrayMath_ISPC_Enabled = false;
-#elif UE_BUILD_SHIPPING
-const bool bAudio_FloatArrayMath_ISPC_Enabled = true;
-#else
-extern bool bAudio_FloatArrayMath_ISPC_Enabled;
-#endif
