Restore lost #endif
#rnx

diff --git a/Engine/Source/Runtime/SignalProcessing/Public/DSP/BufferVectorOperations.h b/Engine/Source/Runtime/SignalProcessing/Public/DSP/BufferVectorOperations.h
--- a/Engine/Source/Runtime/SignalProcessing/Public/DSP/BufferVectorOperations.h
+++ b/Engine/Source/Runtime/SignalProcessing/Public/DSP/BufferVectorOperations.h
@@ -7,6 +7,7 @@
 
 #if UE_ENABLE_INCLUDE_ORDER_DEPRECATED_IN_5_1 
 #include "CoreMinimal.h"
+#endif
 
 namespace Audio
 {
