Make AUDIO_BUFFER_ALIGNMENT definition consistent with AudioResampler.h changes in CL# 21463620
#rnx

diff --git a/Engine/Source/Runtime/SignalProcessing/Public/DSP/AlignedBuffer.h b/Engine/Source/Runtime/SignalProcessing/Public/DSP/AlignedBuffer.h
--- a/Engine/Source/Runtime/SignalProcessing/Public/DSP/AlignedBuffer.h
+++ b/Engine/Source/Runtime/SignalProcessing/Public/DSP/AlignedBuffer.h
@@ -5,12 +5,9 @@
 #include "Containers/Array.h"
 #include "Containers/ContainerAllocationPolicies.h"
 
-#if PLATFORM_SWITCH
-// Switch uses page alignment for submitted buffers
-#define AUDIO_BUFFER_ALIGNMENT 4096
-#else
+// TODO: Move BufferVectorOperations AUDIO_BUFFER_ALIGNMENT
+// define to more central location and reference here.
 #define AUDIO_BUFFER_ALIGNMENT 16
-#endif
 
 #define AUDIO_SIMD_BYTE_ALIGNMENT (16)
 #define AUDIO_NUM_FLOATS_PER_VECTOR_REGISTER (4)
