Re-adding and deprecating audio signal processing functions that had been removed.
#rnx

diff --git a/Engine/Source/Runtime/SignalProcessing/Private/BufferVectorOperations.cpp b/Engine/Source/Runtime/SignalProcessing/Private/BufferVectorOperations.cpp
--- a/Engine/Source/Runtime/SignalProcessing/Private/BufferVectorOperations.cpp
+++ b/Engine/Source/Runtime/SignalProcessing/Private/BufferVectorOperations.cpp
@@ -5,6 +5,221 @@
 
 namespace Audio
 {
+	/* Sets a values to zero if value is denormal. Denormal numbers significantly slow down floating point operations. */
+	void BufferUnderflowClampFast(FAlignedFloatBuffer& InOutBuffer)
+	{
+		ArrayUnderflowClamp(InOutBuffer);
+	}
+	
+	/* Sets a values to zero if value is denormal. Denormal numbers significantly slow down floating point operations. */
+	void BufferUnderflowClampFast(float* RESTRICT InOutBuffer, const int32 InNum)
+	{
+		ArrayUnderflowClamp(MakeArrayView<float>(InOutBuffer, InNum));
+	}
+
+	/* Clamps values in the buffer to be between InMinValue and InMaxValue */
+	void BufferRangeClampFast(FAlignedFloatBuffer& InOutBuffer, float InMinValue, float InMaxValue)
+	{
+		ArrayRangeClamp(InOutBuffer, InMinValue, InMaxValue);
+	}
+
+	/* Clamps values in the buffer to be between InMinValue and InMaxValue */
+	void BufferRangeClampFast(float* RESTRICT InOutBuffer, const int32 InNum, float InMinValue, float InMaxValue)
+	{
+		ArrayRangeClamp(MakeArrayView<float>(InOutBuffer, InNum), InMinValue, InMaxValue);
+	}
+
+	void BufferMultiplyByConstant(const FAlignedFloatBuffer& InFloatBuffer, float InValue, FAlignedFloatBuffer& OutFloatBuffer)
+	{
+		ArrayMultiplyByConstant(InFloatBuffer, InValue, OutFloatBuffer);
+	}
+
+	void BufferMultiplyByConstant(const float* RESTRICT InFloatBuffer, float InValue, float* RESTRICT OutFloatBuffer, const int32 InNumSamples)
+	{
+		ArrayMultiplyByConstant(MakeArrayView<const float>(InFloatBuffer, InNumSamples), InValue, MakeArrayView<float>(OutFloatBuffer, InNumSamples));
+	}
+
+	void MultiplyBufferByConstantInPlace(FAlignedFloatBuffer& InBuffer, float InGain)
+	{
+		ArrayMultiplyByConstantInPlace(InBuffer, InGain);
+	}
+
+	void MultiplyBufferByConstantInPlace(float* RESTRICT InBuffer, int32 NumSamples, float InGain)
+	{
+		ArrayMultiplyByConstantInPlace(MakeArrayView<float>(InBuffer, NumSamples), InGain);
+	}
+
+	// Adds a constant to a buffer (useful for DC offset removal)
+	void AddConstantToBufferInplace(FAlignedFloatBuffer& InBuffer, float InConstant)
+	{
+		ArrayAddConstantInplace(InBuffer, InConstant);
+	}
+
+	void AddConstantToBufferInplace(float* RESTRICT InBuffer, int32 NumSamples, float InConstant)
+	{
+		ArrayAddConstantInplace(MakeArrayView<float>(InBuffer, NumSamples), InConstant);
+	}
+
+	void BufferSetToConstantInplace(FAlignedFloatBuffer& InBuffer, float InConstant)
+	{
+		ArraySetToConstantInplace(InBuffer, InConstant);
+	}
+
+	void BufferSetToConstantInplace(float* RESTRICT InBuffer, int32 NumSamples, float InConstant)
+	{
+		ArraySetToConstantInplace(MakeArrayView<float>(InBuffer, NumSamples), InConstant);
+	}
+
+	/* Performs an element-wise weighted sum OutputBuffer = (InBuffer1 x InGain1) + (InBuffer2 x InGain2) */
+	void BufferWeightedSumFast(const FAlignedFloatBuffer& InBuffer1, float InGain1, const FAlignedFloatBuffer& InBuffer2, float InGain2, FAlignedFloatBuffer& OutBuffer)
+	{
+		ArrayWeightedSum(InBuffer1, InGain1, InBuffer2, InGain2, OutBuffer);
+	}
+
+	/* Performs an element-wise weighted sum OutputBuffer = (InBuffer1 x InGain1) + InBuffer2 */
+	void BufferWeightedSumFast(const FAlignedFloatBuffer& InBuffer1, float InGain1, const FAlignedFloatBuffer& InBuffer2, FAlignedFloatBuffer& OutBuffer)
+	{
+		ArrayWeightedSum(InBuffer1, InGain1, InBuffer2, OutBuffer);
+	}
+
+	/* Performs an element-wise weighted sum OutputBuffer = (InBuffer1 x InGain1) + (InBuffer2 x InGain2) */
+	void BufferWeightedSumFast(const float* RESTRICT InBuffer1, float InGain1, const float* RESTRICT InBuffer2, float InGain2, float* RESTRICT OutBuffer, int32 InNum)
+	{
+		ArrayWeightedSum(MakeArrayView<const float>(InBuffer1, InNum), InGain1, MakeArrayView<const float>(InBuffer2, InNum), InGain2, MakeArrayView<float>(OutBuffer, InNum));
+	}
+
+	/* Performs an element-wise weighted sum OutputBuffer = (InBuffer1 x InGain1) + InBuffer2 */
+	void BufferWeightedSumFast(const float* RESTRICT InBuffer1, float InGain1, const float* RESTRICT InBuffer2, float* RESTRICT OutBuffer, int32 InNum)
+	{
+		ArrayWeightedSum(MakeArrayView<const float>(InBuffer1, InNum), InGain1, MakeArrayView<const float>(InBuffer2, InNum), MakeArrayView<float>(OutBuffer, InNum));
+	}
+
+	void FadeBufferFast(FAlignedFloatBuffer& OutFloatBuffer, const float StartValue, const float EndValue)
+	{
+		ArrayFade(OutFloatBuffer, StartValue, EndValue);
+	}
+
+	void FadeBufferFast(float* RESTRICT OutFloatBuffer, int32 NumSamples, const float StartValue, const float EndValue)
+	{
+		ArrayFade(MakeArrayView<float>(OutFloatBuffer, NumSamples), StartValue, EndValue);
+	}
+
+	void MixInBufferFast(const FAlignedFloatBuffer& InFloatBuffer, FAlignedFloatBuffer& BufferToSumTo, const float Gain)
+	{
+		ArrayMixIn(InFloatBuffer, BufferToSumTo, Gain);
+	}
+
+	void MixInBufferFast(const float* RESTRICT InFloatBuffer, float* RESTRICT BufferToSumTo, int32 NumSamples, const float Gain)
+	{
+		ArrayMixIn(MakeArrayView<const float>(InFloatBuffer, NumSamples), MakeArrayView<float>(BufferToSumTo, NumSamples), Gain);
+	}
+
+	void MixInBufferFast(const FAlignedFloatBuffer& InFloatBuffer, FAlignedFloatBuffer& BufferToSumTo)
+	{
+		ArrayMixIn(InFloatBuffer, BufferToSumTo);
+	}
+
+	void MixInBufferFast(const float* RESTRICT InFloatBuffer, float* RESTRICT BufferToSumTo, int32 NumSamples)
+	{
+		ArrayMixIn(MakeArrayView<const float>(InFloatBuffer, NumSamples), MakeArrayView<float>(BufferToSumTo, NumSamples));
+	}
+
+	void MixInBufferFast(const FAlignedFloatBuffer& InFloatBuffer, FAlignedFloatBuffer& BufferToSumTo, const float StartGain, const float EndGain)
+	{
+		ArrayMixIn(InFloatBuffer, BufferToSumTo, StartGain, EndGain);
+	}
+
+	void MixInBufferFast(const float* RESTRICT InFloatBuffer, float* RESTRICT BufferToSumTo, int32 NumSamples, const float StartGain, const float EndGain)
+	{
+		ArrayMixIn(MakeArrayView<const float>(InFloatBuffer, NumSamples), MakeArrayView<float>(BufferToSumTo, NumSamples), StartGain, EndGain);
+	}
+
+	/* Subtracts two buffers together element-wise. */
+	void BufferSubtractFast(const FAlignedFloatBuffer& InMinuend, const FAlignedFloatBuffer& InSubtrahend, FAlignedFloatBuffer& OutputBuffer)
+	{
+		ArraySubtract(InMinuend, InSubtrahend, OutputBuffer);
+	}
+
+	/* Subtracts two buffers together element-wise. */
+	void BufferSubtractFast(const float* RESTRICT InMinuend, const float* RESTRICT InSubtrahend, float* RESTRICT OutBuffer, int32 InNum)
+	{
+		ArraySubtract(MakeArrayView<const float>(InMinuend, InNum), MakeArrayView<const float>(InSubtrahend, InNum), MakeArrayView<float>(OutBuffer, InNum));
+	}
+
+	/* Performs element-wise in-place subtraction placing the result in the subtrahend. InOutSubtrahend = InMinuend - InOutSubtrahend */
+	void BufferSubtractInPlace1Fast(const FAlignedFloatBuffer& InMinuend, FAlignedFloatBuffer& InOutSubtrahend)
+	{
+		ArraySubtractInPlace1(InMinuend, InOutSubtrahend);
+	}
+
+	/* Performs element-wise in-place subtraction placing the result in the subtrahend. InOutSubtrahend = InMinuend - InOutSubtrahend */
+	void BufferSubtractInPlace1Fast(const float* RESTRICT InMinuend, float* RESTRICT InOutSubtrahend, int32 InNum)
+	{
+		ArraySubtractInPlace1(MakeArrayView<const float>(InMinuend, InNum), MakeArrayView<float>(InOutSubtrahend, InNum));
+	}
+	
+	/* Performs element-wise in-place subtraction placing the result in the minuend. InOutMinuend = InOutMinuend - InSubtrahend */
+	void BufferSubtractInPlace2Fast(FAlignedFloatBuffer& InOutMinuend, const FAlignedFloatBuffer& InSubtrahend)
+	{
+		ArraySubtractInPlace2(InOutMinuend, InSubtrahend);
+	}
+
+	/* Performs element-wise in-place subtraction placing the result in the minuend. InOutMinuend = InOutMinuend - InSubtrahend */
+	void BufferSubtractInPlace2Fast(float* RESTRICT InOutMinuend, const float* RESTRICT InSubtrahend, int32 InNum)
+	{
+		ArraySubtractInPlace2(MakeArrayView<float>(InOutMinuend, InNum), MakeArrayView<const float>(InSubtrahend, InNum));
+	}
+
+	void SumBuffers(const FAlignedFloatBuffer& InFloatBuffer1, const FAlignedFloatBuffer& InFloatBuffer2, FAlignedFloatBuffer& OutputBuffer)
+	{
+		ArraySum(InFloatBuffer1, InFloatBuffer2, OutputBuffer);
+	}
+
+	void SumBuffers(const float* RESTRICT InFloatBuffer1, const float* RESTRICT InFloatBuffer2, float* RESTRICT OutputBuffer, int32 NumSamples)
+	{
+		ArraySum(MakeArrayView<const float>(InFloatBuffer1, NumSamples), MakeArrayView<const float>(InFloatBuffer2, NumSamples), MakeArrayView<float>(OutputBuffer, NumSamples));
+	}
+
+	void MultiplyBuffersInPlace(const FAlignedFloatBuffer& InFloatBuffer, FAlignedFloatBuffer& BufferToMultiply)
+	{
+		ArrayMultiplyInPlace(InFloatBuffer, BufferToMultiply);
+	}
+
+	void MultiplyBuffersInPlace(const float* RESTRICT InFloatBuffer, float* RESTRICT BufferToMultiply, int32 NumSamples)
+	{
+		ArrayMultiplyInPlace(MakeArrayView<const float>(InFloatBuffer, NumSamples), MakeArrayView<float>(BufferToMultiply, NumSamples));
+	}
+
+	float GetMagnitude(const FAlignedFloatBuffer& Buffer)
+	{
+		return ArrayGetMagnitude(Buffer);
+	}
+
+	float GetMagnitude(const float* RESTRICT Buffer, int32 NumSamples)
+	{
+		return ArrayGetMagnitude(MakeArrayView<const float>(Buffer, NumSamples));
+	}
+
+	float BufferGetAverageValue(const FAlignedFloatBuffer& Buffer)
+	{
+		return ArrayGetAverageValue(Buffer);
+	}
+
+	float BufferGetAverageValue(const float* RESTRICT Buffer, int32 NumSamples)
+	{
+		return ArrayGetAverageValue(MakeArrayView<const float>(Buffer, NumSamples));
+	}
+
+	float BufferGetAverageAbsValue(const FAlignedFloatBuffer& Buffer)
+	{
+		return ArrayGetAverageAbsValue(Buffer);
+	}
+
+	float BufferGetAverageAbsValue(const float* RESTRICT Buffer, int32 NumSamples)
+	{
+		return ArrayGetAverageAbsValue(MakeArrayView<const float>(Buffer, NumSamples));
+	}
+
 	/**
 	 * CHANNEL MIXING OPERATIONS:
 	 * To understand these functions, it's best that you have prior experience reading SIMD code.
diff --git a/Engine/Source/Runtime/SignalProcessing/Public/DSP/BufferVectorOperations.h b/Engine/Source/Runtime/SignalProcessing/Public/DSP/BufferVectorOperations.h
--- a/Engine/Source/Runtime/SignalProcessing/Public/DSP/BufferVectorOperations.h
+++ b/Engine/Source/Runtime/SignalProcessing/Public/DSP/BufferVectorOperations.h
@@ -4,6 +4,7 @@
 
 #include "DSP/AlignedBuffer.h"
 #include "Containers/Array.h"
+#include "Misc/CoreMiscDefines.h"
 
 #if UE_ENABLE_INCLUDE_ORDER_DEPRECATED_IN_5_1 
 #include "CoreMinimal.h"
@@ -11,6 +12,153 @@
 
 namespace Audio
 {
+	/* Sets a values to zero if value is denormal. Denormal numbers significantly slow down floating point operations. */
+	UE_DEPRECATED(5.1, "Use equivalent function in FloatArrayMath.h")
+	SIGNALPROCESSING_API void BufferUnderflowClampFast(FAlignedFloatBuffer& InOutBuffer);
+
+	/* Sets a values to zero if value is denormal. Denormal numbers significantly slow down floating point operations. */
+	UE_DEPRECATED(5.1, "Use equivalent function in FloatArrayMath.h")
+	SIGNALPROCESSING_API void BufferUnderflowClampFast(float* RESTRICT InOutBuffer, const int32 InNum);
+	
+	/* Clamps the values in a buffer between a min and max value.*/
+	UE_DEPRECATED(5.1, "Use equivalent function in FloatArrayMath.h")
+	SIGNALPROCESSING_API void BufferRangeClampFast(FAlignedFloatBuffer& InOutBuffer, float InMinValue, float InMaxValue);
+
+	UE_DEPRECATED(5.1, "Use equivalent function in FloatArrayMath.h")
+	SIGNALPROCESSING_API void BufferRangeClampFast(float* RESTRICT InOutBuffer, const int32 InNum, float InMinValue, float InMaxValue);
+	
+	/** Multiplies the input aligned float buffer with the given value. */
+	UE_DEPRECATED(5.1, "Use equivalent function in FloatArrayMath.h")
+	SIGNALPROCESSING_API void BufferMultiplyByConstant(const FAlignedFloatBuffer& InFloatBuffer, float InValue, FAlignedFloatBuffer& OutFloatBuffer);
+	UE_DEPRECATED(5.1, "Use equivalent function in FloatArrayMath.h")
+	SIGNALPROCESSING_API void BufferMultiplyByConstant(const float* RESTRICT InFloatBuffer, float InValue, float* RESTRICT OutFloatBuffer, const int32 InNumSamples);
+	UE_DEPRECATED(5.1, "Use equivalent function in FloatArrayMath.h")
+	SIGNALPROCESSING_API void BufferMultiplyByConstant(const FAlignedFloatBuffer& InFloatBuffer, float InValue, FAlignedFloatBuffer& OutFloatBuffer);
+
+	/** Similar to BufferMultiplyByConstant, but (a) assumes a buffer length divisible by 4 and (b) performs the multiply in place. */
+	UE_DEPRECATED(5.1, "Use equivalent function in FloatArrayMath.h")
+	SIGNALPROCESSING_API void MultiplyBufferByConstantInPlace(FAlignedFloatBuffer& InBuffer, float InGain);
+	UE_DEPRECATED(5.1, "Use equivalent function in FloatArrayMath.h")
+	SIGNALPROCESSING_API void MultiplyBufferByConstantInPlace(float* RESTRICT InBuffer, int32 NumSamples, float InGain);
+
+	/** Adds a constant to a buffer (useful for DC offset removal) */
+	UE_DEPRECATED(5.1, "Use equivalent function in FloatArrayMath.h")
+	SIGNALPROCESSING_API void AddConstantToBufferInplace(FAlignedFloatBuffer& InBuffer, float Constant);
+	UE_DEPRECATED(5.1, "Use equivalent function in FloatArrayMath.h")
+	SIGNALPROCESSING_API void AddConstantToBufferInplace(float* RESTRICT InBuffer, int32 NumSamples, float Constant);
+
+	/** Sets a constant to a buffer (useful for DC offset application) */
+	UE_DEPRECATED(5.1, "Use equivalent function in FloatArrayMath.h")
+	SIGNALPROCESSING_API void BufferSetToConstantInplace(FAlignedFloatBuffer& InBuffer, float Constant);
+	UE_DEPRECATED(5.1, "Use equivalent function in FloatArrayMath.h")
+	SIGNALPROCESSING_API void BufferSetToConstantInplace(float* RESTRICT InBuffer, int32 NumSamples, float Constant);
+
+	/* Performs an element-wise weighted sum OutputBuffer = (InBuffer1 x InGain1) + (InBuffer2 x InGain2) */
+	UE_DEPRECATED(5.1, "Use equivalent function in FloatArrayMath.h")
+	SIGNALPROCESSING_API void BufferWeightedSumFast(const FAlignedFloatBuffer& InBuffer1, float InGain1, const FAlignedFloatBuffer& InBuffer2, float InGain2, FAlignedFloatBuffer& OutBuffer);
+
+	/* Performs an element-wise weighted sum OutputBuffer = (InBuffer1 x InGain1) + InBuffer2 */
+	UE_DEPRECATED(5.1, "Use equivalent function in FloatArrayMath.h")
+	SIGNALPROCESSING_API void BufferWeightedSumFast(const FAlignedFloatBuffer& InBuffer1, float InGain1, const FAlignedFloatBuffer& InBuffer2, FAlignedFloatBuffer& OutBuffer);
+
+	/* Performs an element-wise weighted sum OutputBuffer = (InBuffer1 x InGain1) + (InBuffer2 x InGain2) */
+	UE_DEPRECATED(5.1, "Use equivalent function in FloatArrayMath.h")
+	SIGNALPROCESSING_API void BufferWeightedSumFast(const float* RESTRICT InBuffer1, float InGain1, const float* RESTRICT InBuffer2, float InGain2, float* RESTRICT OutBuffer, int32 InNum);
+
+	/* Performs an element-wise weighted sum OutputBuffer = (InBuffer1 x InGain1) + InBuffer2 */
+	UE_DEPRECATED(5.1, "Use equivalent function in FloatArrayMath.h")
+	SIGNALPROCESSING_API void BufferWeightedSumFast(const float* RESTRICT InBuffer1, float InGain1, const float* RESTRICT InBuffer2, float* RESTRICT OutBuffer, int32 InNum);
+	UE_DEPRECATED(5.1, "Use equivalent function in FloatArrayMath.h")
+	SIGNALPROCESSING_API void MultiplyBufferByConstantInPlace(FAlignedFloatBuffer& InBuffer, float InGain);
+	UE_DEPRECATED(5.1, "Use equivalent function in FloatArrayMath.h")
+	SIGNALPROCESSING_API void MultiplyBufferByConstantInPlace(float* RESTRICT InBuffer, int32 NumSamples, float InGain);
+
+	/* Takes a float buffer and quickly interpolates it's gain from StartValue to EndValue. */
+	/* This operation completely ignores channel counts, so avoid using this function on buffers that are not mono, stereo or quad */
+	/* if the buffer needs to fade all channels uniformly. */
+	UE_DEPRECATED(5.1, "Use equivalent function in FloatArrayMath.h")
+	SIGNALPROCESSING_API void FadeBufferFast(FAlignedFloatBuffer& OutFloatBuffer, const float StartValue, const float EndValue);
+	UE_DEPRECATED(5.1, "Use equivalent function in FloatArrayMath.h")
+	SIGNALPROCESSING_API void FadeBufferFast(float* RESTRICT OutFloatBuffer, int32 NumSamples, const float StartValue, const float EndValue);
+
+	/** Takes buffer InFloatBuffer, optionally multiplies it by Gain, and adds it to BufferToSumTo. */
+	UE_DEPRECATED(5.1, "Use equivalent function in FloatArrayMath.h")
+	SIGNALPROCESSING_API void MixInBufferFast(const FAlignedFloatBuffer& InFloatBuffer, FAlignedFloatBuffer& BufferToSumTo, const float Gain);
+	UE_DEPRECATED(5.1, "Use equivalent function in FloatArrayMath.h")
+	SIGNALPROCESSING_API void MixInBufferFast(const float* RESTRICT InFloatBuffer, float* RESTRICT BufferToSumTo, int32 NumSamples, const float Gain);
+	UE_DEPRECATED(5.1, "Use equivalent function in FloatArrayMath.h")
+	SIGNALPROCESSING_API void MixInBufferFast(const FAlignedFloatBuffer& InFloatBuffer, FAlignedFloatBuffer& BufferToSumTo, const float StartGain, const float EndGain);
+	UE_DEPRECATED(5.1, "Use equivalent function in FloatArrayMath.h")
+	SIGNALPROCESSING_API void MixInBufferFast(const float* RESTRICT InFloatBuffer, float* RESTRICT BufferToSumTo, int32 NumSamples, const float StartGain, const float EndGain);
+	UE_DEPRECATED(5.1, "Use equivalent function in FloatArrayMath.h")
+	SIGNALPROCESSING_API void MixInBufferFast(const FAlignedFloatBuffer& InFloatBuffer, FAlignedFloatBuffer& BufferToSumTo);
+	UE_DEPRECATED(5.1, "Use equivalent function in FloatArrayMath.h")
+	SIGNALPROCESSING_API void MixInBufferFast(const float* RESTRICT InFloatBuffer, float* RESTRICT BufferToSumTo, int32 NumSamples);
+	UE_DEPRECATED(5.1, "Use equivalent function in FloatArrayMath.h")
+
+	/* Subtracts two buffers together element-wise. */
+	UE_DEPRECATED(5.1, "Use equivalent function in FloatArrayMath.h")
+	SIGNALPROCESSING_API void BufferSubtractFast(const FAlignedFloatBuffer& InMinuend, const FAlignedFloatBuffer& InSubtrahend, FAlignedFloatBuffer& OutputBuffer);
+
+	/* Subtracts two buffers together element-wise. */
+	UE_DEPRECATED(5.1, "Use equivalent function in FloatArrayMath.h")
+	SIGNALPROCESSING_API void BufferSubtractFast(const float* RESTRICT InMinuend, const float* RESTRICT InSubtrahend, float* RESTRICT OutputBuffer, int32 NumSamples);
+
+	/* Performs element-wise in-place subtraction placing the result in the subtrahend. InOutSubtrahend = InMinuend - InOutSubtrahend */
+	UE_DEPRECATED(5.1, "Use equivalent function in FloatArrayMath.h")
+	SIGNALPROCESSING_API void BufferSubtractInPlace1Fast(const FAlignedFloatBuffer& InMinuend, FAlignedFloatBuffer& InOutSubtrahend);
+
+	/* Performs element-wise in-place subtraction placing the result in the subtrahend. InOutSubtrahend = InMinuend - InOutSubtrahend */
+	UE_DEPRECATED(5.1, "Use equivalent function in FloatArrayMath.h")
+	SIGNALPROCESSING_API void BufferSubtractInPlace1Fast(const float* RESTRICT InMinuend, float* RESTRICT InOutSubtrahend, int32 NumSamples);
+
+
+	/* Performs element-wise in-place subtraction placing the result in the minuend. InOutMinuend = InOutMinuend - InSubtrahend */
+	UE_DEPRECATED(5.1, "Use equivalent function in FloatArrayMath.h")
+	SIGNALPROCESSING_API void BufferSubtractInPlace2Fast(FAlignedFloatBuffer& InOutMinuend, const FAlignedFloatBuffer& InSubtrahend);
+
+	/* Performs element-wise in-place subtraction placing the result in the minuend. InOutMinuend = InOutMinuend - InSubtrahend */
+	UE_DEPRECATED(5.1, "Use equivalent function in FloatArrayMath.h")
+	SIGNALPROCESSING_API void BufferSubtractInPlace2Fast(float* RESTRICT InOutMinuend, const float* RESTRICT InSubtrahend, int32 NumSamples);
+
+	/** This version of MixInBufferFast will fade from StartGain to EndGain. */
+	UE_DEPRECATED(5.1, "Use equivalent function in FloatArrayMath.h")
+	SIGNALPROCESSING_API void MixInBufferFast(const FAlignedFloatBuffer& InFloatBuffer, FAlignedFloatBuffer& BufferToSumTo, const float StartGain, const float EndGain);
+	UE_DEPRECATED(5.1, "Use equivalent function in FloatArrayMath.h")
+	SIGNALPROCESSING_API void MixInBufferFast(const float* RESTRICT InFloatBuffer, float* RESTRICT BufferToSumTo, int32 NumSamples, const float StartGain, const float EndGain);
+
+	/** Sums two buffers together and places the result in the resulting buffer. */
+	UE_DEPRECATED(5.1, "Use equivalent function in FloatArrayMath.h")
+	SIGNALPROCESSING_API void SumBuffers(const FAlignedFloatBuffer& InFloatBuffer1, const FAlignedFloatBuffer& InFloatBuffer2, FAlignedFloatBuffer& OutputBuffer);
+	UE_DEPRECATED(5.1, "Use equivalent function in FloatArrayMath.h")
+	SIGNALPROCESSING_API void SumBuffers(const float* RESTRICT InFloatBuffer1, const float* RESTRICT InFloatBuffer2, float* RESTRICT OutputBuffer, int32 NumSamples);
+
+	/** Multiply the second buffer in place by the first buffer. */
+	UE_DEPRECATED(5.1, "Use equivalent function in FloatArrayMath.h")
+	SIGNALPROCESSING_API void MultiplyBuffersInPlace(const FAlignedFloatBuffer& InFloatBuffer, FAlignedFloatBuffer& BufferToMultiply);
+	UE_DEPRECATED(5.1, "Use equivalent function in FloatArrayMath.h")
+	SIGNALPROCESSING_API void MultiplyBuffersInPlace(const float* RESTRICT InFloatBuffer, float* RESTRICT BufferToMultiply, int32 NumSamples);
+
+	/** CHANNEL-AGNOSTIC ANALYSIS OPERATIONS */
+
+	/** Takes an audio buffer and returns the magnitude across that buffer. */
+	UE_DEPRECATED(5.1, "Use equivalent function in FloatArrayMath.h")
+	SIGNALPROCESSING_API float GetMagnitude(const FAlignedFloatBuffer& Buffer);
+	UE_DEPRECATED(5.1, "Use equivalent function in FloatArrayMath.h")
+	SIGNALPROCESSING_API float GetMagnitude(const float* RESTRICT Buffer, int32 NumSamples);
+
+	/** Takes an audio buffer and gets the average amplitude across that buffer. */
+	UE_DEPRECATED(5.1, "Use equivalent function in FloatArrayMath.h")
+	SIGNALPROCESSING_API float BufferGetAverageValue(const FAlignedFloatBuffer& Buffer);
+	UE_DEPRECATED(5.1, "Use equivalent function in FloatArrayMath.h")
+	SIGNALPROCESSING_API float BufferGetAverageValue(const float* RESTRICT Buffer, int32 NumSamples);
+
+	/** Takes an audio buffer and gets the average absolute amplitude across that buffer. */
+	UE_DEPRECATED(5.1, "Use equivalent function in FloatArrayMath.h")
+	SIGNALPROCESSING_API float BufferGetAverageAbsValue(const FAlignedFloatBuffer& Buffer);
+	UE_DEPRECATED(5.1, "Use equivalent function in FloatArrayMath.h")
+	SIGNALPROCESSING_API float BufferGetAverageAbsValue(const float* RESTRICT Buffer, int32 NumSamples);
+
 	/** CHANNEL-SPECIFIC OPERATIONS */
 
 	/** Takes a 2 channel interleaved buffer and applies Gains to it. Gains is expected to point to a 2 float long buffer.
