Deprecate AudioFFT in favor of optimized and hardware accelerated version
Deprecating voice processing as is not tested or needed. Current voice processing needs are handled in EOS with WebRTC.

diff --git a/Engine/Plugins/Runtime/AudioSynesthesia/Source/AudioSynesthesiaCore/Private/ConstantQAnalyzer.cpp b/Engine/Plugins/Runtime/AudioSynesthesia/Source/AudioSynesthesiaCore/Private/ConstantQAnalyzer.cpp
--- a/Engine/Plugins/Runtime/AudioSynesthesia/Source/AudioSynesthesiaCore/Private/ConstantQAnalyzer.cpp
+++ b/Engine/Plugins/Runtime/AudioSynesthesia/Source/AudioSynesthesiaCore/Private/ConstantQAnalyzer.cpp
@@ -5,34 +5,44 @@
 #include "DSP/ConstantQ.h"
 #include "DSP/AudioFFT.h"
 #include "DSP/FloatArrayMath.h"
-#include "DSP/BufferVectorOperations.h"
+#include "DSP/FFTAlgorithm.h"
 
 namespace Audio
 {
 	FConstantQAnalyzer::FConstantQAnalyzer(const FConstantQAnalyzerSettings& InSettings, const float InSampleRate)
-	:	Settings(InSettings)
-	,	SampleRate(InSampleRate)
-	,	NumUsefulFFTBins((Settings.FFTSize / 2) + 1)
-    ,   Window(Settings.WindowType, Settings.FFTSize, 1, false)
+	: Settings(InSettings)
+	, SampleRate(InSampleRate)
+	, ActualFFTSize(0)
+	, NumUsefulFFTBins(0)
+    , Window(Settings.WindowType, Settings.FFTSize, 1, false)
 	{
 		// Need FFTSize atleast 8 to support optimized operations.
 		check(Settings.FFTSize >= 8);
 		check(SampleRate > 0.f);
 
-		CQTTransform = NewPseudoConstantQKernelTransform(Settings, Settings.FFTSize, SampleRate);
+		// Create FFT
+		FFFTSettings FFTSettings;
+		FFTSettings.Log2Size = CeilLog2(Settings.FFTSize);
+		FFTSettings.bArrays128BitAligned = true;
+		FFTSettings.bEnableHardwareAcceleration = true;
 
-		// Resize internal buffers
-		WindowedSamples.Reset(Settings.FFTSize);
-		WindowedSamples.AddUninitialized(Settings.FFTSize);
+		ActualFFTSize = 1 << FFTSettings.Log2Size;
+		NumUsefulFFTBins = (ActualFFTSize / 2) + 1;
 
-		FFTOutputRealData.Reset(Settings.FFTSize);
-		FFTOutputRealData.AddUninitialized(Settings.FFTSize);
+		checkf(FFFTFactory::AreFFTSettingsSupported(FFTSettings), TEXT("No fft algorithm supports fft settings."));
+		FFT = FFFTFactory::NewFFTAlgorithm(FFTSettings);
 
-		FFTOutputImagData.Reset(Settings.FFTSize);
-		FFTOutputImagData.AddUninitialized(Settings.FFTSize);
+		// Create CQT kernel
+		CQTTransform = NewPseudoConstantQKernelTransform(Settings, ActualFFTSize, SampleRate);
 
-		SpectrumBuffer.Reset(NumUsefulFFTBins);
-		SpectrumBuffer.AddUninitialized(NumUsefulFFTBins);
+		if (FFT.IsValid())
+		{
+			// Size internal buffers
+			WindowedSamples.AddZeroed(FFT->NumInputFloats()); // Zero samples to apply zero buffer in case actual FFT is larger than provided settings.
+
+			ComplexSpectrum.AddUninitialized(FFT->NumOutputFloats());
+			RealSpectrum.AddUninitialized(ComplexSpectrum.Num() / 2);
+		}
 	}
 
 	void FConstantQAnalyzer::CalculateCQT(const float* InSamples, TArray<float>& OutCQT)
@@ -41,36 +51,48 @@ namespace Audio
 		FMemory::Memcpy(WindowedSamples.GetData(), InSamples, sizeof(float) * Settings.FFTSize);
 		Window.ApplyToBuffer(WindowedSamples.GetData());
 
-		// Take FFT of windowed samples
-		const FFTTimeDomainData TimeData = {WindowedSamples.GetData(), Settings.FFTSize};
-		FFTFreqDomainData FreqData = {FFTOutputRealData.GetData(), FFTOutputImagData.GetData()};
-
-		PerformFFT(TimeData, FreqData);
+		if (FFT.IsValid())
+		{
+			FFT->ForwardRealToComplex(WindowedSamples.GetData(), ComplexSpectrum.GetData());
 
-		// Calculate spectrum from FFT output
-		ComputeSpectrum(Settings.SpectrumType, FreqData, Settings.FFTSize, SpectrumBuffer);
+			ArrayComplexToPower(ComplexSpectrum, RealSpectrum);
 		
-		// Convert spectrum to CQT
-		CQTTransform->TransformArray(SpectrumBuffer, OutCQT);
+			ScalePowerSpectrumInPlace(ActualFFTSize, FFT->ForwardScaling(), EFFTScaling::None, RealSpectrum);
 
-		// Apply decibel scaling if appropriate.
-		if (EConstantQScaling::Decibel == Settings.Scaling)
-		{
-			// Convert to dB
-			switch (Settings.SpectrumType)
+			if (ESpectrumType::MagnitudeSpectrum == Settings.SpectrumType)
+			{
+				// Convert to magnitude
+				ArraySqrtInPlace(RealSpectrum);
+			}
+			else
+			{
+				// Spectrum type should only be magnitude or power. If not true,
+				// then there is a possible missing case coverage. 
+				check(ESpectrumType::PowerSpectrum == Settings.SpectrumType);
+			}
+
+			// Convert spectrum to CQT
+			CQTTransform->TransformArray(RealSpectrum, OutCQT);
+
+			// Apply decibel scaling if appropriate.
+			if (EConstantQScaling::Decibel == Settings.Scaling)
 			{
-				case ESpectrumType::MagnitudeSpectrum:
-					ArrayMagnitudeToDecibelInPlace(OutCQT, -90.f);
-					break;
-
-				case ESpectrumType::PowerSpectrum:
-					ArrayPowerToDecibelInPlace(OutCQT, -90.f);
-					break;
-
-				default:
-					check(false);
-					//checkf(false, TEXT("Unhandled ESpectrumType %s"), GETENUMSTRING(ESpectrumType, Settings.SpectrumType));
-					ArrayPowerToDecibelInPlace(OutCQT, -90.f);
+				// Convert to dB
+				switch (Settings.SpectrumType)
+				{
+					case ESpectrumType::MagnitudeSpectrum:
+						ArrayMagnitudeToDecibelInPlace(OutCQT, -90.f);
+						break;
+
+					case ESpectrumType::PowerSpectrum:
+						ArrayPowerToDecibelInPlace(OutCQT, -90.f);
+						break;
+
+					default:
+						check(false);
+						//checkf(false, TEXT("Unhandled ESpectrumType %s"), GETENUMSTRING(ESpectrumType, Settings.SpectrumType));
+						ArrayPowerToDecibelInPlace(OutCQT, -90.f);
+				}
 			}
 		}
 	}
diff --git a/Engine/Plugins/Runtime/AudioSynesthesia/Source/AudioSynesthesiaCore/Private/LoudnessAnalyzer.cpp b/Engine/Plugins/Runtime/AudioSynesthesia/Source/AudioSynesthesiaCore/Private/LoudnessAnalyzer.cpp
--- a/Engine/Plugins/Runtime/AudioSynesthesia/Source/AudioSynesthesiaCore/Private/LoudnessAnalyzer.cpp
+++ b/Engine/Plugins/Runtime/AudioSynesthesia/Source/AudioSynesthesiaCore/Private/LoudnessAnalyzer.cpp
@@ -3,10 +3,11 @@
 
 #include "LoudnessAnalyzer.h"
 #include "AudioMixer.h"
-#include "DSP/BufferVectorOperations.h"
+#include "DSP/AlignedBuffer.h"
 #include "DSP/DeinterleaveView.h"
+#include "DSP/FloatArrayMath.h"
 
-namespace 
+namespace LoudnessAnalyzerPrivate
 {
 	static const float Log10Scale = 1.0f / FMath::Loge(10.f);
 }
@@ -57,108 +58,112 @@ namespace Audio
 
     FLoudnessAnalyzer::FLoudnessAnalyzer(float InSampleRate, const FLoudnessAnalyzerSettings& InSettings)
     :   Settings(InSettings)
-    ,   EnergyScale(1.0f)
-    ,   MinFreqIndex(0)
-    ,   MaxFreqIndex(0)
-    ,   FFTFreqSize((Settings.FFTSize / 2) + 1)
     ,   Window(Settings.WindowType, Settings.FFTSize, 1, false)
     {
 		checkf(FMath::IsPowerOfTwo(Settings.FFTSize), TEXT("FFT size must be a power of two"));
-		checkf(Settings.FFTSize >= 4, TEXT("FFT size must be atleast 4"))
+		checkf(Settings.FFTSize >= 8, TEXT("FFT size must be atleast 8"))
         check(InSampleRate > 0);
 
-        // Determine which fft bins will be used for calculating loudness
-        const float FreqPerBin = InSampleRate / Settings.FFTSize;
-        MinFreqIndex = static_cast<int32>(FMath::RoundHalfFromZero(Settings.MinAnalysisFrequency / FreqPerBin));
-        MinFreqIndex = FMath::Max(0, MinFreqIndex);
-        MaxFreqIndex = static_cast<int32>(FMath::RoundHalfFromZero(Settings.MaxAnalysisFrequency / FreqPerBin)) + 1;
-        MaxFreqIndex = FMath::Min(MaxFreqIndex, FFTFreqSize);
+		// Create FFT
+		FFFTSettings FFTSettings;
+		FFTSettings.Log2Size = CeilLog2(Settings.FFTSize);
+		FFTSettings.bArrays128BitAligned = true;
+		FFTSettings.bEnableHardwareAcceleration = true;
 
-        // Resize internal buffers
-        WindowedSamples.Reset(Settings.FFTSize);
-        WindowedSamples.AddUninitialized(Settings.FFTSize);
+		ActualFFTSize = 1 << FFTSettings.Log2Size;
 
-        FFTOutputRealData.Reset(Settings.FFTSize);
-        FFTOutputRealData.AddUninitialized(Settings.FFTSize);
-
-        FFTOutputImagData.Reset(Settings.FFTSize);
-        FFTOutputImagData.AddUninitialized(Settings.FFTSize);
-
-        CurveWeights.Reset(FFTFreqSize);
-        CurveWeights.AddUninitialized(FFTFreqSize);
-
-        PowerSpectrumBuffer.Reset(FFTFreqSize);
-        PowerSpectrumBuffer.AddUninitialized(FFTFreqSize);
-
-        float(*LoudnessMapFn)(const float);
-        switch (Settings.LoudnessCurveType)
-        {
-            case ELoudnessCurveType::A:
-                LoudnessMapFn = GetEqualLoudnessAWeightForFrequency;
-                break;
-
-            case ELoudnessCurveType::B:
-                LoudnessMapFn = GetEqualLoudnessBWeightForFrequency;
-                break;
-
-            case ELoudnessCurveType::C:
-                LoudnessMapFn = GetEqualLoudnessCWeightForFrequency;
-                break;
-
-            case ELoudnessCurveType::D:
-                LoudnessMapFn = GetEqualLoudnessDWeightForFrequency;
-                break;
-
-            case ELoudnessCurveType::None:
-            default:
-                LoudnessMapFn = GetEqualLoudnessNoneWeightForFrequency;
-                break;
-        }
-
-        /* Calculate equal loudness weighting curve */
-
-        // Convention is to normalize 1kHz to have a loudness of 1.
-        float EqualLoudnessNorm = 1.0f / LoudnessMapFn(1000.0f);
+		checkf(FFFTFactory::AreFFTSettingsSupported(FFTSettings), TEXT("No fft algorithm supports fft settings."));
+		FFT = FFFTFactory::NewFFTAlgorithm(FFTSettings);
+		if (FFT.IsValid())
+		{
+			// Size internal buffers
+			WindowedSamples.AddZeroed(FFT->NumInputFloats()); // Zero samples to apply zero buffer in case actual FFT is larger than provided settings.
+			ComplexSpectrum.AddUninitialized(FFT->NumOutputFloats());
+			RealSpectrum.AddUninitialized(ComplexSpectrum.Num() / 2);
 
-        float* CurveWeightsPtr = CurveWeights.GetData();
+			FFTFreqSize = RealSpectrum.Num();
 
-        const float FFTSize = static_cast<float>(Settings.FFTSize);
+			// Determine which fft bins will be used for calculating loudness
+			const float FreqPerBin = InSampleRate / FFT->NumOutputFloats();
+			MinFreqIndex = static_cast<int32>(FMath::RoundHalfFromZero(Settings.MinAnalysisFrequency / FreqPerBin));
+			MinFreqIndex = FMath::Max(0, MinFreqIndex);
+			MaxFreqIndex = static_cast<int32>(FMath::RoundHalfFromZero(Settings.MaxAnalysisFrequency / FreqPerBin)) + 1;
+			MaxFreqIndex = FMath::Min(MaxFreqIndex, FFTFreqSize);
 
-        for (int32 i = 0; i < FFTFreqSize; i++)
-        {
-            const float Freq = static_cast<float>(i) * InSampleRate / FFTSize;
-			float Weighting = EqualLoudnessNorm;
+			CurveWeights.Reset(FFTFreqSize);
+			CurveWeights.AddUninitialized(FFTFreqSize);
 
+			float(*LoudnessMapFn)(const float);
 			switch (Settings.LoudnessCurveType)
 			{
 				case ELoudnessCurveType::A:
-					Weighting *= GetEqualLoudnessAWeightForFrequency(Freq);
+					LoudnessMapFn = GetEqualLoudnessAWeightForFrequency;
 					break;
 
 				case ELoudnessCurveType::B:
-					Weighting *= GetEqualLoudnessBWeightForFrequency(Freq);
+					LoudnessMapFn = GetEqualLoudnessBWeightForFrequency;
 					break;
 
 				case ELoudnessCurveType::C:
-					Weighting *= GetEqualLoudnessCWeightForFrequency(Freq);
+					LoudnessMapFn = GetEqualLoudnessCWeightForFrequency;
 					break;
 
 				case ELoudnessCurveType::D:
-					Weighting *= GetEqualLoudnessDWeightForFrequency(Freq);
+					LoudnessMapFn = GetEqualLoudnessDWeightForFrequency;
 					break;
 
 				case ELoudnessCurveType::None:
 				default:
-					Weighting *= GetEqualLoudnessNoneWeightForFrequency(Freq);
+					LoudnessMapFn = GetEqualLoudnessNoneWeightForFrequency;
 					break;
 			}
 
-            // Curve designed for magnitude domain, but applied to power spectrum. The curve is squared to be applied to power spectrum.
-            CurveWeightsPtr[i] = Weighting * Weighting;
-        }
+			/* Calculate equal loudness weighting curve */
+
+			// Convention is to normalize 1kHz to have a loudness of 1.
+			float EqualLoudnessNorm = 1.0f / LoudnessMapFn(1000.0f);
+
+			float* CurveWeightsPtr = CurveWeights.GetData();
 
-		// Normalize by FFT Scaling (1 / FFTSize) and by number of samples (1 / WindowSize) to calculate mean squared value.
-		EnergyScale = 1.f / static_cast<float>(Settings.FFTSize * Settings.FFTSize);
+			const float FFTSize = static_cast<float>(ActualFFTSize);
+
+			for (int32 i = 0; i < FFTFreqSize; i++)
+			{
+				const float Freq = static_cast<float>(i) * InSampleRate / FFTSize;
+				float Weighting = EqualLoudnessNorm;
+
+				switch (Settings.LoudnessCurveType)
+				{
+					case ELoudnessCurveType::A:
+						Weighting *= GetEqualLoudnessAWeightForFrequency(Freq);
+						break;
+
+					case ELoudnessCurveType::B:
+						Weighting *= GetEqualLoudnessBWeightForFrequency(Freq);
+						break;
+
+					case ELoudnessCurveType::C:
+						Weighting *= GetEqualLoudnessCWeightForFrequency(Freq);
+						break;
+
+					case ELoudnessCurveType::D:
+						Weighting *= GetEqualLoudnessDWeightForFrequency(Freq);
+						break;
+
+					case ELoudnessCurveType::None:
+					default:
+						Weighting *= GetEqualLoudnessNoneWeightForFrequency(Freq);
+						break;
+				}
+
+				// Curve designed for magnitude domain, but applied to power spectrum. The curve is squared to be applied to power spectrum.
+				CurveWeightsPtr[i] = Weighting * Weighting;
+			}
+
+			// Normalize by FFT Scaling (1 / FFTSize) and by number of samples (1 / WindowSize) to calculate mean squared value.
+			EnergyScale = GetPowerSpectrumScaling(ActualFFTSize, FFT->ForwardScaling(), EFFTScaling::None);
+			EnergyScale *= 1.f / static_cast<float>(Settings.FFTSize);
+		}
     }
 
     const FLoudnessAnalyzerSettings& FLoudnessAnalyzer::GetSettings() const
@@ -170,31 +175,28 @@ namespace Audio
     {
 		check(InView.Num() == Settings.FFTSize);
 
-		// Copy input samples and apply window
-        FMemory::Memcpy(WindowedSamples.GetData(), InView.GetData(), sizeof(float) * Settings.FFTSize);
-        Window.ApplyToBuffer(WindowedSamples.GetData());
-
-		// Take FFT of windowed samples
-        const FFTTimeDomainData TimeData = {WindowedSamples.GetData(), Settings.FFTSize};
-        FFTFreqDomainData FreqData = {FFTOutputRealData.GetData(), FFTOutputImagData.GetData()};
+		float Total = 0.f;
 
-        PerformFFT(TimeData, FreqData);
+		if (FFT.IsValid())
+		{
+			// Copy input samples and apply window
+			FMemory::Memcpy(WindowedSamples.GetData(), InView.GetData(), sizeof(float) * Settings.FFTSize);
+			Window.ApplyToBuffer(WindowedSamples.GetData());
 
-		// Accumulate energy as weighted sum of power spectrum in frequency analysis range.
-        const float* FFTOutputRealPtr = FFTOutputRealData.GetData();
-        const float* FFTOutputImagPtr = FFTOutputImagData.GetData();
-        float* PowerSpectrumPtr = PowerSpectrumBuffer.GetData();
-		const float* CurveWeightsPtr = CurveWeights.GetData();
+			FFT->ForwardRealToComplex(WindowedSamples.GetData(), ComplexSpectrum.GetData());
 
-		float Total = 0.0f;
-        for (int32 i = MinFreqIndex; i < MaxFreqIndex; i++)
-        {
-            PowerSpectrumPtr[i] = FFTOutputRealPtr[i] * FFTOutputRealPtr[i] + FFTOutputImagPtr[i] * FFTOutputImagPtr[i];
-			Total += PowerSpectrumPtr[i] * CurveWeightsPtr[i];
-        }
+			ArrayComplexToPower(ComplexSpectrum, RealSpectrum);
+		
+			const float* PowerSpectrumPtr = RealSpectrum.GetData();
+			const float* CurveWeightsPtr = CurveWeights.GetData();
+			for (int32 i = MinFreqIndex; i < MaxFreqIndex; i++)
+			{
+				Total += PowerSpectrumPtr[i] * CurveWeightsPtr[i];
+			}
 
-		// Normalize by FFT Scaling and by number of samples to calculate mean squared value.
-		Total *= EnergyScale;
+			// Normalize by FFT Scaling and by number of samples to calculate mean squared value.
+			Total *= EnergyScale;
+		}
 
         return Total;
     }
@@ -207,7 +209,7 @@ namespace Audio
 
 	float FLoudnessAnalyzer::ConvertPerceptualEnergyToLoudness(float InPerceptualEnergy)
 	{
-        return 10.0f * FMath::Loge(InPerceptualEnergy) * Log10Scale;
+        return 10.0f * FMath::Loge(InPerceptualEnergy) * LoudnessAnalyzerPrivate::Log10Scale;
 	}
 
 	/****************************************************************************/
diff --git a/Engine/Plugins/Runtime/AudioSynesthesia/Source/AudioSynesthesiaCore/Private/OnsetAnalyzer.cpp b/Engine/Plugins/Runtime/AudioSynesthesia/Source/AudioSynesthesiaCore/Private/OnsetAnalyzer.cpp
--- a/Engine/Plugins/Runtime/AudioSynesthesia/Source/AudioSynesthesiaCore/Private/OnsetAnalyzer.cpp
+++ b/Engine/Plugins/Runtime/AudioSynesthesia/Source/AudioSynesthesiaCore/Private/OnsetAnalyzer.cpp
@@ -3,6 +3,7 @@
 #include "OnsetAnalyzer.h"
 #include "PeakPicker.h"
 #include "DSP/FloatArrayMath.h"
+#include "DSP/FFTAlgorithm.h"
 
 namespace Audio
 {
@@ -21,24 +22,40 @@ namespace Audio
 		check(Settings.MelSettings.NumBands > 0);
 		check(Settings.NumWindowFrames <= Settings.FFTSize);
 
-		// intialize buffers dealing with audio window
-		WorkingBuffer.AddUninitialized(Settings.NumWindowFrames);
+		FFFTSettings FFTSettings;
+		FFTSettings.Log2Size = CeilLog2(Settings.FFTSize);
+		FFTSettings.bArrays128BitAligned = true;
+		FFTSettings.bEnableHardwareAcceleration = true;
 
-		// initialize buffers dealing with fft input / output
-		WindowedSamples.AddZeroed(Settings.FFTSize);
-		FFTOutputRealData.AddUninitialized(Settings.FFTSize);
-		FFTOutputImagData.AddUninitialized(Settings.FFTSize);
-		FFTSpectrumBuffer.AddUninitialized(Settings.FFTSize);
-		
-		// intialize buffers dealing with perceptual spectrum 
-		PreviousSpectra.AddDefaulted(Settings.ComparisonLag);
-		MelSpectrum.AddUninitialized(Settings.MelSettings.NumBands);
-		SpectrumDifference.AddUninitialized(Settings.MelSettings.NumBands);
+		ActualFFTSize = 1 << FFTSettings.Log2Size;
+
+		checkf(FFFTFactory::AreFFTSettingsSupported(FFTSettings), TEXT("No fft algorithm supports fft settings."));
+		FFT = FFFTFactory::NewFFTAlgorithm(FFTSettings);
+
+		if (FFT.IsValid())
+		{
+			// intialize buffers dealing with audio window
+			WorkingBuffer.AddUninitialized(Settings.NumWindowFrames);
+
+			// initialize buffers dealing with fft input / output
+			WindowedSamples.AddZeroed(FFT->NumInputFloats());
+			ComplexSpectrum.AddUninitialized(FFT->NumOutputFloats());
+			RealSpectrum.AddUninitialized(FFT->NumOutputFloats() / 2);
+			
+			// intialize buffers dealing with perceptual spectrum 
+			PreviousMelSpectra.AddDefaulted(Settings.ComparisonLag);
+			for (FAlignedFloatBuffer& Spectrum : PreviousMelSpectra)
+			{
+				Spectrum.Reset();
+				Spectrum.AddZeroed(Settings.MelSettings.NumBands);
+			}
 
-		// create mel spectrum transform
-		MelTransform = NewMelSpectrumKernelTransform(Settings.MelSettings, Settings.FFTSize, InSampleRate);
+			MelSpectrum.AddUninitialized(Settings.MelSettings.NumBands);
+			MelSpectrumDifference.AddUninitialized(Settings.MelSettings.NumBands);
 
-		check(MelTransform.IsValid());
+			// create mel spectrum transform
+			MelTransform = NewMelSpectrumKernelTransform(Settings.MelSettings, Settings.FFTSize, InSampleRate);
+		}
 	}
 
 	void FOnsetStrengthAnalyzer::CalculateOnsetStrengths(TArrayView<const float> InSamples, TArray<float>& OutEnvelopeStrengths)
@@ -61,9 +78,10 @@ namespace Audio
 	void FOnsetStrengthAnalyzer::Reset()
 	{
 		// Reset internal memory
-		for (TArray<float>& Spectrum : PreviousSpectra)
+		for (FAlignedFloatBuffer& Spectrum : PreviousMelSpectra)
 		{
 			Spectrum.Reset();
+			Spectrum.AddZeroed(Settings.MelSettings.NumBands);
 		}
 		SlidingBuffer.Reset();
 	}
@@ -114,47 +132,44 @@ namespace Audio
 	{
 		float OnsetStrength = 0.f;
 
-		// Copy input samples and apply window
-		FMemory::Memcpy(WindowedSamples.GetData(), InSamples.GetData(), sizeof(float) * Settings.NumWindowFrames);
-		Window.ApplyToBuffer(WindowedSamples.GetData());
-
-		// Take FFT of windowed samples
-		const FFTTimeDomainData TimeData = {WindowedSamples.GetData(), Settings.FFTSize};
-		FFTFreqDomainData FreqData = {FFTOutputRealData.GetData(), FFTOutputImagData.GetData()};
-
-		PerformFFT(TimeData, FreqData);
+		if (FFT.IsValid() && MelTransform.IsValid())
+		{
+			// Copy input samples and apply window
+			FMemory::Memcpy(WindowedSamples.GetData(), InSamples.GetData(), sizeof(float) * Settings.NumWindowFrames);
+			Window.ApplyToBuffer(WindowedSamples.GetData());
 
-		// Calculate spectrum from FFT output
-		ComputePowerSpectrum(FreqData, Settings.FFTSize, FFTSpectrumBuffer);
+			FFT->ForwardRealToComplex(WindowedSamples.GetData(), ComplexSpectrum.GetData());
+			ArrayComplexToPower(ComplexSpectrum, RealSpectrum);
 		
-		// Convert spectrum to CQT
-		MelTransform->TransformArray(FFTSpectrumBuffer, MelSpectrum);
+			ScalePowerSpectrumInPlace(ActualFFTSize, FFT->ForwardScaling(), EFFTScaling::None, RealSpectrum);
+		
+			// Convert spectrum to mel spectrum
+			MelTransform->TransformArray(RealSpectrum, MelSpectrum);
 
-		// Apply decibel scaling 
-		ArrayPowerToDecibelInPlace(MelSpectrum, -90.f);
+			// Apply decibel scaling 
+			ArrayPowerToDecibelInPlace(MelSpectrum, -90.f);
 
-		// Clamp to noise floor
-		ArrayClampMinInPlace(MelSpectrum, Settings.NoiseFloorDb);
+			// Clamp to noise floor
+			ArrayClampMinInPlace(MelSpectrum, Settings.NoiseFloorDb);
 
-		// Onset strength is mean(abs(diff(spectrum - previous_spectrum)))
-		TArray<float>& PreviousSpectrum = PreviousSpectra[LagSpectraIndex];
-		if (PreviousSpectrum.Num() > 0)
-		{
-			ArraySubtract(MelSpectrum, PreviousSpectrum, SpectrumDifference);
+			// Onset strength is mean(abs(diff(spectrum - previous_spectrum)))
+			FAlignedFloatBuffer& PreviousSpectrum = PreviousMelSpectra[LagSpectraIndex];
+			if (PreviousSpectrum.Num() > 0)
+			{
+				ArraySubtract(MelSpectrum, PreviousSpectrum, MelSpectrumDifference);
 
-			// Half wave rectify
-			ArrayClampMinInPlace(SpectrumDifference, 0.f);
+				// Half wave rectify
+				ArrayClampMinInPlace(MelSpectrumDifference, 0.f);
 
-			ArrayMean(SpectrumDifference, OnsetStrength);
-		}
+				ArrayMean(MelSpectrumDifference, OnsetStrength);
+			}
 
-		// Save this spectrum for later use
-		PreviousSpectrum.Reset(Settings.MelSettings.NumBands);
-		PreviousSpectrum.AddUninitialized(Settings.MelSettings.NumBands);
-		FMemory::Memcpy(PreviousSpectrum.GetData(), MelSpectrum.GetData(), sizeof(float) * Settings.MelSettings.NumBands);
+			// Save this spectrum for later use
+			FMemory::Memcpy(PreviousSpectrum.GetData(), MelSpectrum.GetData(), sizeof(float) * Settings.MelSettings.NumBands);
 
-		// Increment index for comparison spectra
-		LagSpectraIndex = (LagSpectraIndex + 1) % Settings.ComparisonLag;
+			// Increment index for comparison spectra
+			LagSpectraIndex = (LagSpectraIndex + 1) % Settings.ComparisonLag;
+		}
 
 		return OnsetStrength;
 	}
diff --git a/Engine/Plugins/Runtime/AudioSynesthesia/Source/AudioSynesthesiaCore/Public/ConstantQAnalyzer.h b/Engine/Plugins/Runtime/AudioSynesthesia/Source/AudioSynesthesiaCore/Public/ConstantQAnalyzer.h
--- a/Engine/Plugins/Runtime/AudioSynesthesia/Source/AudioSynesthesiaCore/Public/ConstantQAnalyzer.h
+++ b/Engine/Plugins/Runtime/AudioSynesthesia/Source/AudioSynesthesiaCore/Public/ConstantQAnalyzer.h
@@ -2,8 +2,9 @@
 
 #pragma once
 
-#include "DSP/ConstantQ.h"
+#include "DSP/AlignedBuffer.h"
 #include "DSP/AudioFFT.h"
+#include "DSP/ConstantQ.h"
 
 namespace Audio
 {
@@ -77,16 +78,19 @@ namespace Audio
 
 		const float SampleRate;
 
+		int32 ActualFFTSize;
+
 		/** Only the first (FFTSize / 2) + 1 FFT bins are used since FFT of real values is conjugate symmetric. */
-		const int32 NumUsefulFFTBins;
+		int32 NumUsefulFFTBins;
 
 		/** Scale factor to normalize between different FFT Sizes */
 		FWindow Window;
 
 		FAlignedFloatBuffer WindowedSamples;
-		FAlignedFloatBuffer FFTOutputRealData;
-		FAlignedFloatBuffer FFTOutputImagData;
-		FAlignedFloatBuffer SpectrumBuffer;
+		FAlignedFloatBuffer ComplexSpectrum;
+		FAlignedFloatBuffer RealSpectrum;
+
+		TUniquePtr<IFFTAlgorithm> FFT;
 
 		TUniquePtr<FContiguousSparse2DKernelTransform> CQTTransform;
 	};
diff --git a/Engine/Plugins/Runtime/AudioSynesthesia/Source/AudioSynesthesiaCore/Public/LoudnessAnalyzer.h b/Engine/Plugins/Runtime/AudioSynesthesia/Source/AudioSynesthesiaCore/Public/LoudnessAnalyzer.h
--- a/Engine/Plugins/Runtime/AudioSynesthesia/Source/AudioSynesthesiaCore/Public/LoudnessAnalyzer.h
+++ b/Engine/Plugins/Runtime/AudioSynesthesia/Source/AudioSynesthesiaCore/Public/LoudnessAnalyzer.h
@@ -2,7 +2,7 @@
 
 #pragma once
 
-#include "DSP/BufferVectorOperations.h"
+#include "DSP/AlignedBuffer.h"
 #include "DSP/AudioFFT.h"
 
 namespace Audio
@@ -98,22 +98,23 @@ namespace Audio
             FLoudnessAnalyzerSettings Settings;
 
 			/** Scale factor used to normalize between differing window sizes and fft sizes */
-            float EnergyScale;
+            float EnergyScale = 1.f;
 
-            int32 MinFreqIndex;
-            int32 MaxFreqIndex;
+            int32 MinFreqIndex = 0;
+            int32 MaxFreqIndex = 0;
 
 			/** Number of FFTFrequency bins worth looking at. */
-			int32 FFTFreqSize;
+			int32 FFTFreqSize = 0;
+			int32 ActualFFTSize = 0;
 
             FWindow Window;
+			TUniquePtr<IFFTAlgorithm> FFT;
 
             FAlignedFloatBuffer WindowedSamples;
-            FAlignedFloatBuffer FFTOutputRealData;
-            FAlignedFloatBuffer FFTOutputImagData;
+            FAlignedFloatBuffer ComplexSpectrum;
+            FAlignedFloatBuffer RealSpectrum;
 
             FAlignedFloatBuffer CurveWeights;
-            FAlignedFloatBuffer PowerSpectrumBuffer;
     };
 
 
diff --git a/Engine/Plugins/Runtime/AudioSynesthesia/Source/AudioSynesthesiaCore/Public/OnsetAnalyzer.h b/Engine/Plugins/Runtime/AudioSynesthesia/Source/AudioSynesthesiaCore/Public/OnsetAnalyzer.h
--- a/Engine/Plugins/Runtime/AudioSynesthesia/Source/AudioSynesthesiaCore/Public/OnsetAnalyzer.h
+++ b/Engine/Plugins/Runtime/AudioSynesthesia/Source/AudioSynesthesiaCore/Public/OnsetAnalyzer.h
@@ -2,10 +2,12 @@
 
 #pragma once
 
-#include "PeakPicker.h"
+#include "DSP/AlignedBuffer.h"
+#include "DSP/AudioFFT.h"
+#include "DSP/FFTAlgorithm.h"
 #include "DSP/MelScale.h"
 #include "DSP/SlidingWindow.h"
-#include "DSP/AudioFFT.h"
+#include "PeakPicker.h"
 
 namespace Audio
 {
@@ -71,22 +73,23 @@ namespace Audio
 
 			FOnsetStrengthSettings Settings;
 			int32 LagSpectraIndex;
+			int32 ActualFFTSize;
 
 			TSlidingBuffer<float> SlidingBuffer;
 			FAlignedFloatBuffer WorkingBuffer;
 			FAlignedFloatBuffer WindowedSamples;
-			FAlignedFloatBuffer FFTOutputRealData;
-			FAlignedFloatBuffer FFTOutputImagData;
-			FAlignedFloatBuffer FFTSpectrumBuffer;
+			FAlignedFloatBuffer ComplexSpectrum;
+			FAlignedFloatBuffer RealSpectrum;
 
 			// Array of arrays to hold lag spectra.
-			TArray< TArray<float> > PreviousSpectra;
+			TArray<FAlignedFloatBuffer> PreviousMelSpectra;
 
 			TArray<float> MelSpectrum;
-			TArray<float> SpectrumDifference;
+			TArray<float> MelSpectrumDifference;
 
 			FWindow Window;
 			TUniquePtr<FContiguousSparse2DKernelTransform> MelTransform;
+			TUniquePtr<IFFTAlgorithm> FFT;
 	};
 
 	/** 
diff --git a/Engine/Source/Runtime/SignalProcessing/Private/AudioFFT.cpp b/Engine/Source/Runtime/SignalProcessing/Private/AudioFFT.cpp
--- a/Engine/Source/Runtime/SignalProcessing/Private/AudioFFT.cpp
+++ b/Engine/Source/Runtime/SignalProcessing/Private/AudioFFT.cpp
@@ -173,6 +173,33 @@ namespace Audio
 
 	namespace FFTIntrinsics
 	{
+		float GetScalingExponent(EFFTScaling InScaling)
+		{
+			switch (InScaling)
+			{
+				case EFFTScaling::None:
+					return 0.f;
+
+				case EFFTScaling::MultipliedByFFTSize:
+					return 1.f;
+
+				case EFFTScaling::MultipliedBySqrtFFTSize:
+					return 0.5f;
+
+				case EFFTScaling::DividedByFFTSize:
+					return -1.f;
+
+				case EFFTScaling::DividedBySqrtFFTSize:
+					return -0.5f;
+
+				default:
+				{
+					checkNoEntry();
+					return 0;
+				}
+			}
+		}
+
 		 uint32 NextPowerOf2(uint32 Input)
 		{
 			uint32 Value = 2;
@@ -458,7 +485,7 @@ namespace Audio
 			}
 		}
 
-		void PerformIterativeFFT(const FFTTimeDomainData& InputParams, FFTFreqDomainData& OutputParams)
+		void PerformIterativeFFT(const FFTTimeDomainData_DEPRECATED& InputParams, FFTFreqDomainData_DEPRECATED& OutputParams)
 		{
 			// Separate even and odd elements into real buffer:
 			SeparateIntoCopy(InputParams.Buffer, OutputParams.OutReal, InputParams.NumSamples);
@@ -470,7 +497,7 @@ namespace Audio
 			ComputeButterfliesInPlace(OutputParams.OutReal, OutputParams.OutImag, InputParams.NumSamples);
 		}
 
-		void PerformIterativeIFFT(FFTFreqDomainData& InputParams, FFTTimeDomainData& OutputParams)
+		void PerformIterativeIFFT(FFTFreqDomainData_DEPRECATED& InputParams, FFTTimeDomainData_DEPRECATED& OutputParams)
 		{
 			TArrayView<float> OutImagView(InputParams.OutImag, OutputParams.NumSamples);
 			TArrayView<float> OutputBufferView(OutputParams.Buffer, OutputParams.NumSamples);
@@ -499,7 +526,7 @@ namespace Audio
 #endif
 		}
 
-		void PerformDFT(const FFTTimeDomainData& InputParams, FFTFreqDomainData& OutputParams)
+		void PerformDFT(const FFTTimeDomainData_DEPRECATED& InputParams, FFTFreqDomainData_DEPRECATED& OutputParams)
 		{
 			const float* InputBuffer = InputParams.Buffer;
 			float* OutReal = OutputParams.OutReal;
@@ -524,7 +551,7 @@ namespace Audio
 			}
 		}
 
-		void PerformIDFT(const FFTFreqDomainData& InputParams, FFTTimeDomainData& OutputParams)
+		void PerformIDFT(const FFTFreqDomainData_DEPRECATED& InputParams, FFTTimeDomainData_DEPRECATED& OutputParams)
 		{
 			float* OutputBuffer = OutputParams.Buffer;
 			float* InReal = InputParams.OutReal;
@@ -546,32 +573,262 @@ namespace Audio
 				OutputBuffer[TimeIndex] = RealSum;
 			}
 		}
-	}
 
-	void PerformFFT(const FFTTimeDomainData& InputParams, FFTFreqDomainData& OutputParams)
+		void ComputePowerSpectrumNoScaling(const FFTFreqDomainData_DEPRECATED& InFrequencyData, int32 FFTSize, FAlignedFloatBuffer& OutBuffer)
+		{
+			check((FFTSize % 2) == 0);
+
+			if (FFTSize < 1)
+			{
+				// Can't do anything with a zero sized fft.
+				OutBuffer.Reset(0);
+				return;
+			}
+
+			// Spectrum only calculates values for real positive frequencies. 
+			const int32 NumSpectrumValues = (FFTSize / 2) + 1;
+
+			// Resize output buffer
+			OutBuffer.Reset(NumSpectrumValues);
+			OutBuffer.AddUninitialized(NumSpectrumValues);
+
+			TArrayView<float> OutBufferDataView(OutBuffer.GetData(), NumSpectrumValues);
+			TArrayView<const float> RealDataView(InFrequencyData.OutReal, NumSpectrumValues);
+			TArrayView<const float> ImagDataView(InFrequencyData.OutImag, NumSpectrumValues);
+
+			ArrayComplexToPower(RealDataView, ImagDataView, OutBufferDataView);
+		}
+
+	} // namespace FFTIntrinsic
+
+	namespace AudioFFTDeprecated
 	{
-		int32 FFTMethod = CVarFFTMethod.GetValueOnAnyThread();
-		if (FFTMethod)
+		void CrossCorrelate_DEPRECATED(FrequencyBuffer& FirstBufferFrequencies, FrequencyBuffer& SecondBufferFrequencies, int32 NumSamples, FrequencyBuffer& OutCorrelation);
+		void CrossCorrelate_DEPRECATED(const float* FirstBuffer, const float* SecondBuffer, int32 NumSamples, int32 FFTSize, FrequencyBuffer& FirstBufferFrequencies, FrequencyBuffer& SecondBufferFrequencies, FrequencyBuffer& OutCorrelation);
+		void CrossCorrelate_DEPRECATED(const float* FirstBuffer, const float* SecondBuffer, int32 NumSamples, int32 FFTSize, FrequencyBuffer& OutCorrelation);
+		void CrossCorrelate_DEPRECATED(const float* FirstBuffer, const float* SecondBuffer, int32 NumSamples, int32 FFTSize, float* OutCorrelation, int32 OutCorrelationSamples);
+		void CrossCorrelate_DEPRECATED(FAlignedFloatBuffer& FirstBuffer, FAlignedFloatBuffer& SecondBuffer, FrequencyBuffer& OutCorrelation, bool bZeroPad);
+		void CrossCorrelate_DEPRECATED(FAlignedFloatBuffer& FirstBuffer, FAlignedFloatBuffer& SecondBuffer, FAlignedFloatBuffer& OutCorrelation, bool bZeroPad);
+
+
+		void PerformFFT_DEPRECATED(const FFTTimeDomainData_DEPRECATED& InputParams, FFTFreqDomainData_DEPRECATED& OutputParams)
+		{
+			int32 FFTMethod = CVarFFTMethod.GetValueOnAnyThread();
+			if (FFTMethod)
+			{
+				FFTIntrinsics::PerformDFT(InputParams, OutputParams);
+			}
+			else
+			{
+				FFTIntrinsics::PerformIterativeFFT(InputParams, OutputParams);
+			}
+		}
+
+		void PerformIFFT_DEPRECATED(FFTFreqDomainData_DEPRECATED& InputParams, FFTTimeDomainData_DEPRECATED& OutputParams)
 		{
-			FFTIntrinsics::PerformDFT(InputParams, OutputParams);
+			int32 FFTMethod = CVarFFTMethod.GetValueOnAnyThread();
+			if (FFTMethod)
+			{
+				FFTIntrinsics::PerformIDFT(InputParams, OutputParams);
+			}
+			else
+			{
+				FFTIntrinsics::PerformIterativeIFFT(InputParams, OutputParams);
+			}
 		}
-		else
+
+		void ComputeMagnitudeSpectrum_DEPRECATED(const FFTFreqDomainData_DEPRECATED& InFrequencyData, int32 FFTSize, FAlignedFloatBuffer& OutBuffer)
 		{
-			FFTIntrinsics::PerformIterativeFFT(InputParams, OutputParams);
+			if (FFTSize < 1)
+			{
+				OutBuffer.Reset();
+				return;
+			}
+
+			FFTIntrinsics::ComputePowerSpectrumNoScaling(InFrequencyData, FFTSize, OutBuffer);
+
+			const int32 NumSpectrumValues = OutBuffer.Num();
+			if (NumSpectrumValues < 1)
+			{
+				return;
+			}
+
+			const float FFTScale = 1.f / FMath::Sqrt(static_cast<float>(FFTSize));
+
+			float* OutBufferData = OutBuffer.GetData();
+
+			for (int32 i = 0; i < NumSpectrumValues; i++)
+			{
+				// TODO: Currently no vector sqrt. utilize fmath for now.
+				OutBufferData[i] = FMath::Sqrt(OutBufferData[i]);
+			}
+
+			if (NumSpectrumValues > 1)
+			{
+				TArrayView<float> OutBufferView(OutBufferData, NumSpectrumValues - 1);
+				ArrayMultiplyByConstantInPlace(OutBufferView, FFTScale);
+			}
+			OutBufferData[NumSpectrumValues - 1] *= FFTScale;
 		}
-	}
 
-	void PerformIFFT(FFTFreqDomainData& InputParams, FFTTimeDomainData& OutputParams)
-	{
-		int32 FFTMethod = CVarFFTMethod.GetValueOnAnyThread();
-		if (FFTMethod)
+		void ComputePowerSpectrum_DEPRECATED(const FFTFreqDomainData_DEPRECATED& InFrequencyData, int32 FFTSize, FAlignedFloatBuffer& OutBuffer)
+		{
+			if (FFTSize < 1)
+			{
+				OutBuffer.Reset();
+				return;
+			}
+
+			FFTIntrinsics::ComputePowerSpectrumNoScaling(InFrequencyData, FFTSize, OutBuffer);
+
+			const int32 NumSpectrumValues = OutBuffer.Num();
+			if (NumSpectrumValues < 1)
+			{
+				return;
+			}
+
+			const float FFTScale = 1.f / static_cast<float>(FFTSize);
+
+			float* OutBufferData = OutBuffer.GetData();
+
+			if (NumSpectrumValues > 1)
+			{
+				TArrayView<float> OutBufferDataView(OutBufferData, NumSpectrumValues - 1);
+				ArrayMultiplyByConstantInPlace(OutBufferDataView, FFTScale);
+			}
+			
+			OutBufferData[NumSpectrumValues - 1] *= FFTScale;
+		}
+
+		void CrossCorrelate_DEPRECATED(FAlignedFloatBuffer& FirstBuffer, FAlignedFloatBuffer& SecondBuffer, FAlignedFloatBuffer& OutCorrelation, bool bZeroPad /*= true*/)
+		{
+			FrequencyBuffer OutputCorrelationFrequencies;
+			CrossCorrelate_DEPRECATED(FirstBuffer, SecondBuffer, OutputCorrelationFrequencies, bZeroPad);
+
+			OutCorrelation.Reset();
+			OutCorrelation.AddUninitialized(OutputCorrelationFrequencies.Real.Num());
+
+			// Perform IFFT into OutCorrelation:
+			FFTFreqDomainData_DEPRECATED FreqDomainData = 
+			{
+				OutputCorrelationFrequencies.Real.GetData(),
+				OutputCorrelationFrequencies.Imag.GetData()
+			};
+
+			FFTTimeDomainData_DEPRECATED TimeDomainData =
+			{
+				OutCorrelation.GetData(),
+				OutCorrelation.Num()
+			};
+
+			AudioFFTDeprecated::PerformIFFT_DEPRECATED(FreqDomainData, TimeDomainData);
+		}
+
+		void CrossCorrelate_DEPRECATED(FAlignedFloatBuffer& FirstBuffer, FAlignedFloatBuffer& SecondBuffer, FrequencyBuffer& OutCorrelation, bool bZeroPad /*= true*/)
+		{
+			const int32 NumSamples = FMath::Max(FirstBuffer.Num(), SecondBuffer.Num());
+
+			if (!bZeroPad)
+			{
+				int32 FFTLength = FFTIntrinsics::NextPowerOf2(NumSamples - 1);
+
+				FirstBuffer.AddZeroed(FFTLength - FirstBuffer.Num());
+				SecondBuffer.AddZeroed(FFTLength - SecondBuffer.Num());
+			}
+			else
+			{
+				checkSlow(FirstBuffer.Num() == SecondBuffer.Num() && FMath::IsPowerOfTwo(FirstBuffer.Num()));
+			}
+
+			CrossCorrelate_DEPRECATED(FirstBuffer.GetData(), SecondBuffer.GetData(), NumSamples, FirstBuffer.Num(), OutCorrelation);
+		}
+
+		void CrossCorrelate_DEPRECATED(const float* FirstBuffer, const float* SecondBuffer, int32 NumSamples, int32 FFTSize, float* OutCorrelation, int32 OutCorrelationSamples)
+		{
+			FrequencyBuffer OutputCorrelationFrequencies;
+			CrossCorrelate_DEPRECATED(FirstBuffer, SecondBuffer, NumSamples, FFTSize, OutputCorrelationFrequencies);
+
+			checkSlow(FFTSize == OutCorrelationSamples);
+
+			// Perform IFFT into OutCorrelation:
+			FFTFreqDomainData_DEPRECATED FreqDomainData =
+			{
+				OutputCorrelationFrequencies.Real.GetData(),
+				OutputCorrelationFrequencies.Imag.GetData()
+			};
+
+			FFTTimeDomainData_DEPRECATED TimeDomainData =
+			{
+				OutCorrelation,
+				OutCorrelationSamples
+			};
+
+			AudioFFTDeprecated::PerformIFFT_DEPRECATED(FreqDomainData, TimeDomainData);
+		}
+
+		void CrossCorrelate_DEPRECATED(const float* FirstBuffer, const float* SecondBuffer, int32 NumSamples, int32 FFTSize, FrequencyBuffer& OutCorrelation)
 		{
-			FFTIntrinsics::PerformIDFT(InputParams, OutputParams);
+			FrequencyBuffer FirstBufferFrequencies;
+			FrequencyBuffer SecondBufferFrequencies;
+
+			CrossCorrelate_DEPRECATED(FirstBuffer, SecondBuffer, NumSamples, FFTSize, FirstBufferFrequencies, SecondBufferFrequencies, OutCorrelation);
 		}
-		else
+
+		void CrossCorrelate_DEPRECATED(const float* FirstBuffer, const float* SecondBuffer, int32 NumSamples, int32 FFTSize, FrequencyBuffer& FirstBufferFrequencies, FrequencyBuffer& SecondBufferFrequencies, FrequencyBuffer& OutCorrelation)
 		{
-			FFTIntrinsics::PerformIterativeIFFT(InputParams, OutputParams);
+			checkSlow(FMath::IsPowerOfTwo(FFTSize));
+			OutCorrelation.InitZeroed(FFTSize);
+			
+			// Perform FFT on First Buffer input:
+			FirstBufferFrequencies.InitZeroed(FFTSize);
+			FFTTimeDomainData_DEPRECATED TimeDomainData =
+			{
+				const_cast<float*>(FirstBuffer), // Note: Because we use the same struct for the IFFT as we do the FFT, we can't make this property const.
+				FFTSize
+			};
+
+			FFTFreqDomainData_DEPRECATED FreqDomainData =
+			{
+				FirstBufferFrequencies.Real.GetData(),
+				FirstBufferFrequencies.Imag.GetData()
+			};
+			
+			AudioFFTDeprecated::PerformFFT_DEPRECATED(TimeDomainData, FreqDomainData);
+
+			// Perform FFT on second buffer of input:
+			SecondBufferFrequencies.InitZeroed(FFTSize);
+			TimeDomainData =
+			{
+				const_cast<float*>(SecondBuffer),
+				FFTSize
+			};
+
+			FreqDomainData = 
+			{
+				SecondBufferFrequencies.Real.GetData(),
+				SecondBufferFrequencies.Imag.GetData()
+			};
+
+			AudioFFTDeprecated::PerformFFT_DEPRECATED(TimeDomainData, FreqDomainData);
+
+			CrossCorrelate_DEPRECATED(FirstBufferFrequencies, SecondBufferFrequencies, NumSamples, OutCorrelation);
 		}
+
+		void CrossCorrelate_DEPRECATED(FrequencyBuffer& FirstBufferFrequencies, FrequencyBuffer& SecondBufferFrequencies, int32 NumSamples, FrequencyBuffer& OutCorrelation)
+		{
+			FFTIntrinsics::ComplexMultiplyByConjugate(FirstBufferFrequencies, SecondBufferFrequencies, OutCorrelation);
+		}
+
+	} // Namespace AudioFFTDeprecated
+
+	void PerformFFT(const FFTTimeDomainData_DEPRECATED& InputParams, FFTFreqDomainData_DEPRECATED& OutputParams)
+	{
+		AudioFFTDeprecated::PerformFFT_DEPRECATED(InputParams, OutputParams);
+	}
+
+	void PerformIFFT(FFTFreqDomainData_DEPRECATED& InputParams, FFTTimeDomainData_DEPRECATED& OutputParams)
+	{
+		AudioFFTDeprecated::PerformIFFT_DEPRECATED(InputParams, OutputParams);
 	}
 
 	// Implementation of fft algorithm interface
@@ -579,7 +836,7 @@ namespace Audio
 	{
 			const int32 FFTSize;
 			const int32 NumOutputFFTElements;
-			FFTFreqDomainData FreqDomainData;
+			FFTFreqDomainData_DEPRECATED FreqDomainData;
 
 			FAlignedFloatBuffer FreqRealBuffer;
 			FAlignedFloatBuffer FreqImagBuffer;
@@ -618,8 +875,8 @@ namespace Audio
 			// OutComplex - Array of floats to store output of fourier transform. Must have (FFTSize() + 2) float elements which represent ((FFTSize() / 2) + 1) complex numbers in interleaved format.
 			virtual void ForwardRealToComplex(const float* RESTRICT InReal, float* RESTRICT OutComplex) override
 			{
-				const FFTTimeDomainData TimeDomainData = {const_cast<float*>(InReal), FFTSize};
-				PerformFFT(TimeDomainData, FreqDomainData);
+				const FFTTimeDomainData_DEPRECATED TimeDomainData = {const_cast<float*>(InReal), FFTSize};
+				AudioFFTDeprecated::PerformFFT_DEPRECATED(TimeDomainData, FreqDomainData);
 
 				// Convert FFT output data to interleaved format.
 				int32 OutPos = 0;
@@ -659,8 +916,8 @@ namespace Audio
 					InFreqPos--;
 				}
 
-				FFTTimeDomainData TimeDomainData = {OutReal, FFTSize};
-				PerformIFFT(FreqDomainData, TimeDomainData);
+				FFTTimeDomainData_DEPRECATED TimeDomainData = {OutReal, FFTSize};
+				AudioFFTDeprecated::PerformIFFT_DEPRECATED(FreqDomainData, TimeDomainData);
 			}
 
 			// BatchForwardRealToComplex
@@ -722,239 +979,69 @@ namespace Audio
 		return MakeUnique<FAudioFFTAlgorithm>(FFTSize);
 	}
 
-	void ComputePowerSpectrumNoScaling(const FFTFreqDomainData& InFrequencyData, int32 FFTSize, FAlignedFloatBuffer& OutBuffer)
-	{
-		check((FFTSize % 2) == 0);
-
-		if (FFTSize < 1)
-		{
-			// Can't do anything with a zero sized fft.
-			OutBuffer.Reset(0);
-			return;
-		}
-
-		// Spectrum only calculates values for real positive frequencies. 
-		const int32 NumSpectrumValues = (FFTSize / 2) + 1;
-
-		// Resize output buffer
-		OutBuffer.Reset(NumSpectrumValues);
-		OutBuffer.AddUninitialized(NumSpectrumValues);
-
-		TArrayView<float> OutBufferDataView(OutBuffer.GetData(), NumSpectrumValues);
-		TArrayView<const float> RealDataView(InFrequencyData.OutReal, NumSpectrumValues);
-		TArrayView<const float> ImagDataView(InFrequencyData.OutImag, NumSpectrumValues);
-
-		ArrayComplexToPower(RealDataView, ImagDataView, OutBufferDataView);
-	}
-
-	void ComputePowerSpectrum(const FFTFreqDomainData& InFrequencyData, int32 FFTSize, FAlignedFloatBuffer& OutBuffer)
+	void ComputePowerSpectrum(const FFTFreqDomainData_DEPRECATED& InFrequencyData, int32 FFTSize, FAlignedFloatBuffer& OutBuffer)
 	{
-		if (FFTSize < 1)
-		{
-			OutBuffer.Reset();
-			return;
-		}
-
-		ComputePowerSpectrumNoScaling(InFrequencyData, FFTSize, OutBuffer);
-
-		const int32 NumSpectrumValues = OutBuffer.Num();
-		if (NumSpectrumValues < 1)
-		{
-			return;
-		}
-
-		const float FFTScale = 1.f / static_cast<float>(FFTSize);
-
-		float* OutBufferData = OutBuffer.GetData();
-
-		if (NumSpectrumValues > 1)
-		{
-			TArrayView<float> OutBufferDataView(OutBufferData, NumSpectrumValues - 1);
-			ArrayMultiplyByConstantInPlace(OutBufferDataView, FFTScale);
-		}
-		
-		OutBufferData[NumSpectrumValues - 1] *= FFTScale;
+		AudioFFTDeprecated::ComputePowerSpectrum_DEPRECATED(InFrequencyData, FFTSize, OutBuffer);
 	}
 
-	void ComputeMagnitudeSpectrum(const FFTFreqDomainData& InFrequencyData, int32 FFTSize, FAlignedFloatBuffer& OutBuffer)
+	void ComputeMagnitudeSpectrum(const FFTFreqDomainData_DEPRECATED& InFrequencyData, int32 FFTSize, FAlignedFloatBuffer& OutBuffer)
 	{
-		if (FFTSize < 1)
-		{
-			OutBuffer.Reset();
-			return;
-		}
-
-		ComputePowerSpectrumNoScaling(InFrequencyData, FFTSize, OutBuffer);
-
-		const int32 NumSpectrumValues = OutBuffer.Num();
-		if (NumSpectrumValues < 1)
-		{
-			return;
-		}
-
-		const float FFTScale = 1.f / FMath::Sqrt(static_cast<float>(FFTSize));
-
-		float* OutBufferData = OutBuffer.GetData();
-
-		for (int32 i = 0; i < NumSpectrumValues; i++)
-		{
-			// TODO: Currently no vector sqrt. utilize fmath for now.
-			OutBufferData[i] = FMath::Sqrt(OutBufferData[i]);
-		}
-
-		if (NumSpectrumValues > 1)
-		{
-			TArrayView<float> OutBufferView(OutBufferData, NumSpectrumValues - 1);
-			ArrayMultiplyByConstantInPlace(OutBufferView, FFTScale);
-		}
-		OutBufferData[NumSpectrumValues - 1] *= FFTScale;
+		AudioFFTDeprecated::ComputeMagnitudeSpectrum_DEPRECATED(InFrequencyData, FFTSize, OutBuffer);
 	}
 
-	void ComputeSpectrum(ESpectrumType InSpectrumType, const FFTFreqDomainData& InFrequencyData, int32 FFTSize, FAlignedFloatBuffer& OutBuffer)
+	void ComputeSpectrum(ESpectrumType InSpectrumType, const FFTFreqDomainData_DEPRECATED& InFrequencyData, int32 FFTSize, FAlignedFloatBuffer& OutBuffer)
 	{
 		switch (InSpectrumType)
 		{
 			case ESpectrumType::MagnitudeSpectrum:
-				ComputeMagnitudeSpectrum(InFrequencyData, FFTSize, OutBuffer);
+				AudioFFTDeprecated::ComputeMagnitudeSpectrum_DEPRECATED(InFrequencyData, FFTSize, OutBuffer);
 				break;
 
 			case ESpectrumType::PowerSpectrum:
-				ComputePowerSpectrum(InFrequencyData, FFTSize, OutBuffer);
+				AudioFFTDeprecated::ComputePowerSpectrum_DEPRECATED(InFrequencyData, FFTSize, OutBuffer);
 				break;
 
 			default:
 				checkf(false, TEXT("Unhandled Audio::ESpectrumType"));
 		}
 	}
-	SIGNALPROCESSING_API void CrossCorrelate(FAlignedFloatBuffer& FirstBuffer, FAlignedFloatBuffer& SecondBuffer, FAlignedFloatBuffer& OutCorrelation, bool bZeroPad /*= true*/)
-	{
-		FrequencyBuffer OutputCorrelationFrequencies;
-		CrossCorrelate(FirstBuffer, SecondBuffer, OutputCorrelationFrequencies, bZeroPad);
-
-		OutCorrelation.Reset();
-		OutCorrelation.AddUninitialized(OutputCorrelationFrequencies.Real.Num());
-
-		// Perform IFFT into OutCorrelation:
-		FFTFreqDomainData FreqDomainData = 
-		{
-			OutputCorrelationFrequencies.Real.GetData(),
-			OutputCorrelationFrequencies.Imag.GetData()
-		};
-
-		FFTTimeDomainData TimeDomainData =
-		{
-			OutCorrelation.GetData(),
-			OutCorrelation.Num()
-		};
 
-		PerformIFFT(FreqDomainData, TimeDomainData);
+	void CrossCorrelate(FAlignedFloatBuffer& FirstBuffer, FAlignedFloatBuffer& SecondBuffer, FAlignedFloatBuffer& OutCorrelation, bool bZeroPad /*= true*/)
+	{
+		AudioFFTDeprecated::CrossCorrelate_DEPRECATED(FirstBuffer, SecondBuffer, OutCorrelation, bZeroPad);
 	}
 
-	SIGNALPROCESSING_API void CrossCorrelate(FAlignedFloatBuffer& FirstBuffer, FAlignedFloatBuffer& SecondBuffer, FrequencyBuffer& OutCorrelation, bool bZeroPad /*= true*/)
+	void CrossCorrelate(FAlignedFloatBuffer& FirstBuffer, FAlignedFloatBuffer& SecondBuffer, FrequencyBuffer& OutCorrelation, bool bZeroPad /*= true*/)
 	{
-		const int32 NumSamples = FMath::Max(FirstBuffer.Num(), SecondBuffer.Num());
-
-		if (!bZeroPad)
-		{
-			int32 FFTLength = FFTIntrinsics::NextPowerOf2(NumSamples - 1);
-
-			FirstBuffer.AddZeroed(FFTLength - FirstBuffer.Num());
-			SecondBuffer.AddZeroed(FFTLength - SecondBuffer.Num());
-		}
-		else
-		{
-			checkSlow(FirstBuffer.Num() == SecondBuffer.Num() && FMath::IsPowerOfTwo(FirstBuffer.Num()));
-		}
-
-		CrossCorrelate(FirstBuffer.GetData(), SecondBuffer.GetData(), NumSamples, FirstBuffer.Num(), OutCorrelation);
+		AudioFFTDeprecated::CrossCorrelate_DEPRECATED(FirstBuffer, SecondBuffer, OutCorrelation, bZeroPad);
 	}
 
-	SIGNALPROCESSING_API void CrossCorrelate(const float* FirstBuffer, const float* SecondBuffer, int32 NumSamples, int32 FFTSize, float* OutCorrelation, int32 OutCorrelationSamples)
+	void CrossCorrelate(const float* FirstBuffer, const float* SecondBuffer, int32 NumSamples, int32 FFTSize, float* OutCorrelation, int32 OutCorrelationSamples)
 	{
-		FrequencyBuffer OutputCorrelationFrequencies;
-		CrossCorrelate(FirstBuffer, SecondBuffer, NumSamples, FFTSize, OutputCorrelationFrequencies);
-
-		checkSlow(FFTSize == OutCorrelationSamples);
-
-		// Perform IFFT into OutCorrelation:
-		FFTFreqDomainData FreqDomainData =
-		{
-			OutputCorrelationFrequencies.Real.GetData(),
-			OutputCorrelationFrequencies.Imag.GetData()
-		};
-
-		FFTTimeDomainData TimeDomainData =
-		{
-			OutCorrelation,
-			OutCorrelationSamples
-		};
-
-		PerformIFFT(FreqDomainData, TimeDomainData);
+		AudioFFTDeprecated::CrossCorrelate_DEPRECATED(FirstBuffer, SecondBuffer, NumSamples, FFTSize, OutCorrelation, OutCorrelationSamples);
 	}
 
-	SIGNALPROCESSING_API void CrossCorrelate(const float* FirstBuffer, const float* SecondBuffer, int32 NumSamples, int32 FFTSize, FrequencyBuffer& OutCorrelation)
+	void CrossCorrelate(const float* FirstBuffer, const float* SecondBuffer, int32 NumSamples, int32 FFTSize, FrequencyBuffer& OutCorrelation)
 	{
-		FrequencyBuffer FirstBufferFrequencies;
-		FrequencyBuffer SecondBufferFrequencies;
-
-		CrossCorrelate(FirstBuffer, SecondBuffer, NumSamples, FFTSize, FirstBufferFrequencies, SecondBufferFrequencies, OutCorrelation);
+		AudioFFTDeprecated::CrossCorrelate_DEPRECATED(FirstBuffer, SecondBuffer, NumSamples, FFTSize, OutCorrelation);
 	}
 
-	SIGNALPROCESSING_API void CrossCorrelate(const float* FirstBuffer, const float* SecondBuffer, int32 NumSamples, int32 FFTSize, FrequencyBuffer& FirstBufferFrequencies, FrequencyBuffer& SecondBufferFrequencies, FrequencyBuffer& OutCorrelation)
+	void CrossCorrelate(const float* FirstBuffer, const float* SecondBuffer, int32 NumSamples, int32 FFTSize, FrequencyBuffer& FirstBufferFrequencies, FrequencyBuffer& SecondBufferFrequencies, FrequencyBuffer& OutCorrelation)
 	{
-		checkSlow(FMath::IsPowerOfTwo(FFTSize));
-		OutCorrelation.InitZeroed(FFTSize);
-		
-		// Perform FFT on First Buffer input:
-		FirstBufferFrequencies.InitZeroed(FFTSize);
-		FFTTimeDomainData TimeDomainData =
-		{
-			const_cast<float*>(FirstBuffer), // Note: Because we use the same struct for the IFFT as we do the FFT, we can't make this property const.
-			FFTSize
-		};
-
-		FFTFreqDomainData FreqDomainData =
-		{
-			FirstBufferFrequencies.Real.GetData(),
-			FirstBufferFrequencies.Imag.GetData()
-		};
-		
-		PerformFFT(TimeDomainData, FreqDomainData);
-
-		// Perform FFT on second buffer of input:
-		SecondBufferFrequencies.InitZeroed(FFTSize);
-		TimeDomainData =
-		{
-			const_cast<float*>(SecondBuffer),
-			FFTSize
-		};
-
-		FreqDomainData = 
-		{
-			SecondBufferFrequencies.Real.GetData(),
-			SecondBufferFrequencies.Imag.GetData()
-		};
-
-		PerformFFT(TimeDomainData, FreqDomainData);
-
-		CrossCorrelate(FirstBufferFrequencies, SecondBufferFrequencies, NumSamples, OutCorrelation);
+		AudioFFTDeprecated::CrossCorrelate_DEPRECATED(FirstBuffer, SecondBuffer, NumSamples, FFTSize, FirstBufferFrequencies, SecondBufferFrequencies, OutCorrelation);
 	}
 
-	SIGNALPROCESSING_API void CrossCorrelate(FrequencyBuffer& FirstBufferFrequencies, FrequencyBuffer& SecondBufferFrequencies, int32 NumSamples, FrequencyBuffer& OutCorrelation)
+	void CrossCorrelate(FrequencyBuffer& FirstBufferFrequencies, FrequencyBuffer& SecondBufferFrequencies, int32 NumSamples, FrequencyBuffer& OutCorrelation)
 	{
-		FFTIntrinsics::ComplexMultiplyByConjugate(FirstBufferFrequencies, SecondBufferFrequencies, OutCorrelation);
-
-		// Normalize our output by the length of the signals:
-		//const float NormalizationFactor = 1.0f / NumSamples;
-		//FFTIntrinsics::ComplexMultiplyInPlaceByConstant(OutCorrelation, NormalizationFactor);
+		AudioFFTDeprecated::CrossCorrelate_DEPRECATED(FirstBufferFrequencies, SecondBufferFrequencies, NumSamples, OutCorrelation);
 	}
 
-	FFFTConvolver::FFFTConvolver()
+	FFFTConvolver_DEPRECATED::FFFTConvolver_DEPRECATED()
 		: BlockSize(0)
 	{
 	}
 
-	void FFFTConvolver::ConvolveBlock(float* InputAudio, int32 NumSamples)
+	void FFFTConvolver_DEPRECATED::ConvolveBlock(float* InputAudio, int32 NumSamples)
 	{
 		checkSlow(NumSamples != 0);
 
@@ -965,21 +1052,21 @@ namespace Audio
 		TimeDomainInputBuffer.AddZeroed(FFTSize);
 		FMemory::Memcpy(TimeDomainInputBuffer.GetData(), InputAudio, NumSamples * sizeof(float));
 
-		FFTTimeDomainData TimeDomainData = 
+		FFTTimeDomainData_DEPRECATED TimeDomainData = 
 		{ 
 			TimeDomainInputBuffer.GetData(), // Buffer
 			TimeDomainInputBuffer.Num()      // FFTSize
 		};
 
-		FFTFreqDomainData FreqDomainData = 
+		FFTFreqDomainData_DEPRECATED FreqDomainData = 
 		{
 			InputFrequencies.Real.GetData(),
 			InputFrequencies.Imag.GetData()
 		};
 
-		PerformFFT(TimeDomainData, FreqDomainData);
+		AudioFFTDeprecated::PerformFFT_DEPRECATED(TimeDomainData, FreqDomainData);
 		FFTIntrinsics::ComplexMultiplyInPlace(FilterFrequencies, InputFrequencies);
-		PerformIFFT(FreqDomainData, TimeDomainData);
+		AudioFFTDeprecated::PerformIFFT_DEPRECATED(FreqDomainData, TimeDomainData);
 
 		// Copy back from our temporary buffer to InputAudio:
 		FMemory::Memcpy(InputAudio, TimeDomainInputBuffer.GetData(), NumSamples * sizeof(float));
@@ -993,7 +1080,7 @@ namespace Audio
 		SetCOLABuffer(TimeDomainInputBuffer.GetData() + COLAOffset, COLASize);
 	}
 
-	void FFFTConvolver::SumInCOLABuffer(float* InputAudio, int32 NumSamples)
+	void FFFTConvolver_DEPRECATED::SumInCOLABuffer(float* InputAudio, int32 NumSamples)
 	{
 		TArrayView<const float> COLABufferView(COLABuffer.GetData(), NumSamples);
 		TArrayView<float> InputAudioView(InputAudio, NumSamples);
@@ -1001,13 +1088,13 @@ namespace Audio
 		ArrayMixIn(COLABufferView, InputAudioView);
 	}
 
-	void FFFTConvolver::SetCOLABuffer(float* InAudio, int32 NumSamples)
+	void FFFTConvolver_DEPRECATED::SetCOLABuffer(float* InAudio, int32 NumSamples)
 	{
 		COLABuffer.SetNumUninitialized(NumSamples, false);
 		FMemory::Memcpy(COLABuffer.GetData(), InAudio, NumSamples * sizeof(float));
 	}
 
-	void FFFTConvolver::ProcessAudio(float* InputAudio, int32 NumSamples)
+	void FFFTConvolver_DEPRECATED::ProcessAudio(float* InputAudio, int32 NumSamples)
 	{
 		checkSlow(BlockSize != 0);
 
@@ -1026,7 +1113,7 @@ namespace Audio
 		}
 	}
 
-	void FFFTConvolver::SetFilter(const float* InWindowReal, const float* InWindowImag, int32 FilterSize, int32 FFTSize)
+	void FFFTConvolver_DEPRECATED::SetFilter(const float* InWindowReal, const float* InWindowImag, int32 FilterSize, int32 FFTSize)
 	{
 		checkSlow(FilterSize != 0);
 
@@ -1046,7 +1133,7 @@ namespace Audio
 		COLABuffer.AddZeroed(BlockSize);
 	}
 
-	void FFFTConvolver::SetFilter(const FrequencyBuffer& InFilterFrequencies, int32 FilterSize)
+	void FFFTConvolver_DEPRECATED::SetFilter(const FrequencyBuffer& InFilterFrequencies, int32 FilterSize)
 	{
 		checkSlow(FilterSize != 0);
 
@@ -1066,7 +1153,7 @@ namespace Audio
 		COLABuffer.AddZeroed(BlockSize);
 	}
 
-	void FFFTConvolver::SetFilter(const float* TimeDomainBuffer, int32 FilterSize)
+	void FFFTConvolver_DEPRECATED::SetFilter(const float* TimeDomainBuffer, int32 FilterSize)
 	{
 		checkSlow(FilterSize != 0);
 
@@ -1078,27 +1165,67 @@ namespace Audio
 		TimeDomainInputBuffer.AddZeroed(FilterFFTSize);
 		FMemory::Memcpy(TimeDomainInputBuffer.GetData(), TimeDomainBuffer, FilterSize * sizeof(float));
 
-		FFTTimeDomainData TimeDomainData =
+		FFTTimeDomainData_DEPRECATED TimeDomainData =
 		{
 			TimeDomainInputBuffer.GetData(),
 			TimeDomainInputBuffer.Num()
 		};
 
-		FFTFreqDomainData FreqDomainData =
+		FFTFreqDomainData_DEPRECATED FreqDomainData =
 		{
 			FilterFrequencies.Real.GetData(),
 			FilterFrequencies.Imag.GetData()
 		};
 
-		PerformFFT(TimeDomainData, FreqDomainData);
+		AudioFFTDeprecated::PerformFFT_DEPRECATED(TimeDomainData, FreqDomainData);
 		BlockSize = FilterSize;
 		InputFrequencies.InitZeroed(FilterFFTSize);
 		COLABuffer.Reset(BlockSize);
 		COLABuffer.AddZeroed(BlockSize);
 	}
 
-	void FFFTConvolver::SetFilter(const FAlignedFloatBuffer& TimeDomainBuffer)
+	void FFFTConvolver_DEPRECATED::SetFilter(const FAlignedFloatBuffer& TimeDomainBuffer)
 	{
 		SetFilter(TimeDomainBuffer.GetData(), TimeDomainBuffer.Num());
 	}
+
+	int32 CeilLog2(int32 InNum)
+	{
+		static constexpr int32 MaxValue = 0x40000000;
+		static constexpr int32 One = 1;
+
+		int32 Result = 0;
+		int32 Value = 1;
+
+		while ((Value < InNum) && (Value < MaxValue))
+		{
+			Result++;
+			Value = One << Result;
+		}
+
+		return Result;
+	}
+
+
+	float GetPowerSpectrumScaling(int32 FFTSize, EFFTScaling InCurrentScaling, EFFTScaling InTargetScaling)
+	{
+		if (!ensureMsgf(FFTSize > 0, TEXT("Invalid FFTSize %d"), FFTSize))
+		{
+			return 1.f;
+		}
+
+		const float ScalingExponentDiff = FFTIntrinsics::GetScalingExponent(InTargetScaling) - FFTIntrinsics::GetScalingExponent(InCurrentScaling);
+		return FMath::Pow(static_cast<float>(FFTSize), ScalingExponentDiff * 2.f);
+	}
+
+	void ScalePowerSpectrumInPlace(int32 FFTSize, EFFTScaling InCurrentScaling, EFFTScaling InTargetScaling, TArrayView<float> InPowerSpectrum)
+	{
+		if (InCurrentScaling != InTargetScaling)
+		{
+			const float Scaling = GetPowerSpectrumScaling(FFTSize, InCurrentScaling, InTargetScaling);
+			ArrayMultiplyByConstantInPlace(InPowerSpectrum, Scaling);
+		}
+	}
+
+
 }
diff --git a/Engine/Source/Runtime/SignalProcessing/Private/SignalProcessingModule.cpp b/Engine/Source/Runtime/SignalProcessing/Private/SignalProcessingModule.cpp
--- a/Engine/Source/Runtime/SignalProcessing/Private/SignalProcessingModule.cpp
+++ b/Engine/Source/Runtime/SignalProcessing/Private/SignalProcessingModule.cpp
@@ -13,7 +13,6 @@ namespace Audio
 {
 	class FSignalProcessingModule : public IModuleInterface
 	{
-		FAudioFFTAlgorithmFactory AudioFFTAlgorithmFactory;
 		FVectorFFTFactory VectorFFTAlgorithmFactory;
 
 		FUniformPartitionConvolutionFactory UniformPartitionConvolutionFactory;
@@ -23,7 +22,6 @@ namespace Audio
 		virtual void StartupModule() override
 		{
 			// FFT factories to register
-			IModularFeatures::Get().RegisterModularFeature(IFFTAlgorithmFactory::GetModularFeatureName(), &AudioFFTAlgorithmFactory);
 			IModularFeatures::Get().RegisterModularFeature(IFFTAlgorithmFactory::GetModularFeatureName(), &VectorFFTAlgorithmFactory);
 
 			// Convolution factories to register
@@ -33,7 +31,6 @@ namespace Audio
 		virtual void ShutdownModule() override
 		{
 			// FFT Factories to unregister
-			IModularFeatures::Get().UnregisterModularFeature(IFFTAlgorithmFactory::GetModularFeatureName(), &AudioFFTAlgorithmFactory);
 			IModularFeatures::Get().UnregisterModularFeature(IFFTAlgorithmFactory::GetModularFeatureName(), &VectorFFTAlgorithmFactory);
 
 			// Convolution factories to unregister
diff --git a/Engine/Source/Runtime/SignalProcessing/Private/VoiceProcessing.cpp b/Engine/Source/Runtime/SignalProcessing/Private/VoiceProcessing.cpp
--- a/Engine/Source/Runtime/SignalProcessing/Private/VoiceProcessing.cpp
+++ b/Engine/Source/Runtime/SignalProcessing/Private/VoiceProcessing.cpp
@@ -188,7 +188,7 @@ namespace Audio
 		return FMath::Clamp(UnclampedGain, GainMin, GainMax);
 	}
 
-	FAdaptiveFilter::FAdaptiveFilter(int32 FilterLength, int32 AudioCallbackSize)
+	FAdaptiveFilter_DEPRECATED::FAdaptiveFilter_DEPRECATED(int32 FilterLength, int32 AudioCallbackSize)
 		: WindowSize(FilterLength)
 		, CurrentStepsUntilConvergence(0)
 	{
@@ -208,7 +208,7 @@ namespace Audio
 		Convolver.SetFilter(CurrentWeights, FilterLength);
 	}
 
-	void FAdaptiveFilter::AdaptFilter()
+	void FAdaptiveFilter_DEPRECATED::AdaptFilter()
 	{
 		// If we've suitably converged, we avoid incrementing our weights.
 		if (CurrentStepsUntilConvergence <= 0)
@@ -221,7 +221,7 @@ namespace Audio
 		Convolver.SetFilter(CurrentWeights, WindowSize);
 	}
 
-	void FAdaptiveFilter::SetWeightDeltas(const float* InWeightsReal, const float* InWeightsImag, int32 NumWeights, float InLearningRate)
+	void FAdaptiveFilter_DEPRECATED::SetWeightDeltas(const float* InWeightsReal, const float* InWeightsImag, int32 NumWeights, float InLearningRate)
 	{
 		// ProcessAudio::NumWeights needs to use the same 
 		checkSlow(NumWeights == WeightDeltas.Real.Num());
@@ -252,7 +252,7 @@ namespace Audio
 		}
 	}
 
-	void FAdaptiveFilter::IncrementWeights()
+	void FAdaptiveFilter_DEPRECATED::IncrementWeights()
 	{
 		const int32 NumWeights = CurrentWeights.Real.Num();
 		float* CurrentRealBuffer = CurrentWeights.Real.GetData();
@@ -275,23 +275,23 @@ namespace Audio
 		}
 	}
 
-	void FAdaptiveFilter::ProcessAudio(float* InAudio, int32 NumSamples)
+	void FAdaptiveFilter_DEPRECATED::ProcessAudio(float* InAudio, int32 NumSamples)
 	{
 		AdaptFilter();
 		Convolver.ProcessAudio(InAudio, NumSamples);
 	}
 
-	void FAdaptiveFilter::SetWeights(const FrequencyBuffer& InFilterWeights, int32 FilterLength, float InLearningRate)
+	void FAdaptiveFilter_DEPRECATED::SetWeights(const FrequencyBuffer& InFilterWeights, int32 FilterLength, float InLearningRate)
 	{
 		SetWeightDeltas(InFilterWeights.Real.GetData(), InFilterWeights.Imag.GetData(), InFilterWeights.Real.Num(), InLearningRate);
 	}
 
-	FFDAPFilterComputer::FFDAPFilterComputer()
+	FFDAPFilterComputer_DEPRECATED::FFDAPFilterComputer_DEPRECATED()
 	{
 
 	}
 
-	void FFDAPFilterComputer::GenerateWeights(const float* IncomingSignal, int32 NumIncomingSamples, const float* OutgoingSignal, int32 NumOutgoingSamples, FrequencyBuffer& OutWeights)
+	void FFDAPFilterComputer_DEPRECATED::GenerateWeights(const float* IncomingSignal, int32 NumIncomingSamples, const float* OutgoingSignal, int32 NumOutgoingSamples, FrequencyBuffer& OutWeights)
 	{
 		int32 FFTSize = FFTIntrinsics::NextPowerOf2(NumIncomingSamples + NumOutgoingSamples - 1);
 		
@@ -305,13 +305,15 @@ namespace Audio
 		FMemory::Memcpy(ZeroPaddedOutgoingBuffer.GetData(), OutgoingSignal, NumOutgoingSamples * sizeof(float));
 
 		int32 FilterLength = FMath::Max(NumIncomingSamples, NumOutgoingSamples);
+PRAGMA_DISABLE_DEPRECATION_WARNINGS
 		CrossCorrelate(ZeroPaddedIncomingBuffer.GetData(), ZeroPaddedOutgoingBuffer.GetData(), FilterLength, FFTSize, IncomingFrequencies, OutgoingFrequencies, OutWeights);
+PRAGMA_ENABLE_DEPRECATION_WARNINGS
 
 		// TODO: Perform linear phase conversion on OutWeights.
 
 	}
 
-	FAcousticEchoCancellation::FAcousticEchoCancellation(float InConvergenceRate, int32 CallbackSize, int32 InFilterLength, int32 InFilterUpdateRate /*= 1*/)
+	FAcousticEchoCancellation_DEPRECATED::FAcousticEchoCancellation_DEPRECATED(float InConvergenceRate, int32 CallbackSize, int32 InFilterLength, int32 InFilterUpdateRate /*= 1*/)
 		: AdaptiveFilter(InFilterLength, CallbackSize)
 		, ConvergenceRate(InConvergenceRate)
 		, FilterLength(InFilterLength)
@@ -321,7 +323,7 @@ namespace Audio
 		checkSlow(FMath::IsPowerOfTwo(FilterLength));
 	}
 
-	void FAcousticEchoCancellation::ProcessAudio(float* InAudio, int32 NumSamples)
+	void FAcousticEchoCancellation_DEPRECATED::ProcessAudio(float* InAudio, int32 NumSamples)
 	{
 		checkSlow(FMath::IsPowerOfTwo(NumSamples));
 
@@ -365,12 +367,12 @@ namespace Audio
 		}
 	}
 
-	Audio::FPatchInput FAcousticEchoCancellation::AddNewSignalPatch(int32 ExpectedLatency, float Gain /*= 1.0f*/)
+	Audio::FPatchInput FAcousticEchoCancellation_DEPRECATED::AddNewSignalPatch(int32 ExpectedLatency, float Gain /*= 1.0f*/)
 	{
 		return PatchMixer.AddNewInput(ExpectedLatency, Gain);
 	}
 
-	void FAcousticEchoCancellation::RemoveSignalPatch(const FPatchInput& Patch)
+	void FAcousticEchoCancellation_DEPRECATED::RemoveSignalPatch(const FPatchInput& Patch)
 	{
 		PatchMixer.RemovePatch(Patch);
 	}
diff --git a/Engine/Source/Runtime/SignalProcessing/Public/DSP/AudioFFT.h b/Engine/Source/Runtime/SignalProcessing/Public/DSP/AudioFFT.h
--- a/Engine/Source/Runtime/SignalProcessing/Public/DSP/AudioFFT.h
+++ b/Engine/Source/Runtime/SignalProcessing/Public/DSP/AudioFFT.h
@@ -3,6 +3,7 @@
 #pragma once
 
 #include "CoreMinimal.h"
+#include "Misc/CoreMiscDefines.h"
 #include "DSP/BufferVectorOperations.h"
 #include "DSP/Dsp.h"
 #include "DSP/FFTAlgorithm.h"
@@ -73,27 +74,54 @@ namespace Audio
 		void Generate(int32 NumFrames, int32 NumChannels, bool bIsPeriodic);
 	};
 
-	struct FFTTimeDomainData
+	struct FFTTimeDomainData_DEPRECATED
 	{
 		float* Buffer; // Pointer to a single channel of floats.
 		int32 NumSamples; // Number of samples in InBuffer divided by the number of channels. must be a power of 2.
 	};
 
-	struct FFTFreqDomainData
+	/** AudioFFT Deprecation
+	 *
+	 * The AudioFFT is deprecated in favor of optimized software implementations 
+	 * and hardware implementations. The FFT implemenented in "PerformFFT(...)" 
+	 * has poor CPU performance as it computes the FFT radix weights each time
+	 * "PerformFFT(...)" is executed. 
+	 *
+	 * Better optimized and hardward FFT implementations can be accessed through the 
+	 * FFFTFactory class in FFTAlgorithm.h. That class will attempt to return the 
+	 * best FFT algorithm implementation available for the given FFFTSettings.
+	 */
+
+	struct UE_DEPRECATED(5.1, "Use IFFTAlgorithm and TArray for time domain data") FFTTimeDomainData;
+	struct FFTTimeDomainData : FFTTimeDomainData_DEPRECATED 
+	{
+	};
+
+	struct FFTFreqDomainData_DEPRECATED
 	{
 		// arrays in which real and imaginary values will be populated.
 		float* OutReal; // Should point to an already allocated array of floats that is FFTInputParams::NumSamples long.
 		float* OutImag; // Should point to an already allocated array of floats that is FFTInputParams::NumSamples long.
 	};
 
+	struct UE_DEPRECATED(5.1, "Use IFFTAlgorithm and TArray for time freq domain data") FFTFreqDomainData;
+	struct FFTFreqDomainData : FFTFreqDomainData_DEPRECATED 
+	{
+	};
+
 	// Performs a one-time FFT on a float buffer. Does not support complex signals.
 	// This function assumes that, if you desire a window for your FFT, that window was already
 	// applied to FFTInputParams.InBuffer.
-	SIGNALPROCESSING_API void PerformFFT(const FFTTimeDomainData& InputParams, FFTFreqDomainData& OutputParams);
-	SIGNALPROCESSING_API void PerformIFFT(FFTFreqDomainData& InputParams, FFTTimeDomainData& OutputParams);
+
+ 	UE_DEPRECATED(5.1, "Use FVectorFFTRealToComplex or FFFTFactory instead.")
+	SIGNALPROCESSING_API void PerformFFT(const FFTTimeDomainData_DEPRECATED& InputParams, FFTFreqDomainData_DEPRECATED& OutputParams);
+
+ 	UE_DEPRECATED(5.1, "Use FVectorFFTRealToComplex or FFFTFactory instead.")
+	SIGNALPROCESSING_API void PerformIFFT(FFTFreqDomainData_DEPRECATED& InputParams, FFTTimeDomainData_DEPRECATED& OutputParams);
 
 
 	// FFT Algorithm factory for this FFT implementation
+	class UE_DEPRECATED(5.1, "Use FVectorFFTFactory instead.") FAudioFFTAlgorithmFactory;
 	class SIGNALPROCESSING_API FAudioFFTAlgorithmFactory : public IFFTAlgorithmFactory
 	{
 		public:
@@ -116,6 +144,7 @@ namespace Audio
 	};
 
 	struct FrequencyBuffer
+ 	
 	{
 		FAlignedFloatBuffer Real;
 		FAlignedFloatBuffer Imag;
@@ -147,19 +176,26 @@ namespace Audio
 	// Performs an acyclic FFT correlation on FirstBuffer and Second buffer and stores the output in OutCorrelation.
 	// If bCyclic is false, This function may zero pad FirstBuffer and Second Buffer as needed.
 	// If bCyclic is true, FirstBuffer and SecondBuffer should have the same length, and that length should be a power of two.
+	UE_DEPRECATED(5.1, "Cross correlate is no longer supported.")
 	SIGNALPROCESSING_API void CrossCorrelate(FAlignedFloatBuffer& FirstBuffer, FAlignedFloatBuffer& SecondBuffer, FAlignedFloatBuffer& OutCorrelation, bool bZeroPad = true);
+	UE_DEPRECATED(5.1, "Cross correlate is no longer supported.")
 	SIGNALPROCESSING_API void CrossCorrelate(FAlignedFloatBuffer& FirstBuffer, FAlignedFloatBuffer& SecondBuffer, FrequencyBuffer& OutCorrelation, bool bZeroPad = true);
+	UE_DEPRECATED(5.1, "Cross correlate is no longer supported.")
 	SIGNALPROCESSING_API void CrossCorrelate(const float* FirstBuffer, const float* SecondBuffer, int32 NumSamples, int32 FFTSize, float* OutCorrelation, int32 OutCorrelationSamples);
+	UE_DEPRECATED(5.1, "Cross correlate is no longer supported.")
 	SIGNALPROCESSING_API void CrossCorrelate(const float* FirstBuffer, const float* SecondBuffer, int32 NumSamples, int32 FFTSize, FrequencyBuffer& OutCorrelation);
 
 	// These variations do not allocate any additional memory during the function, provided that the FrequencyBuffers are already allocated.
+	UE_DEPRECATED(5.1, "Cross correlate is no longer supported.")
 	SIGNALPROCESSING_API void CrossCorrelate(const float* FirstBuffer, const float* SecondBuffer, int32 NumSamples, int32 FFTSize, FrequencyBuffer& FirstBufferFrequencies, FrequencyBuffer& SecondBufferFrequencies, FrequencyBuffer& OutCorrelation);
+	UE_DEPRECATED(5.1, "Cross correlate is no longer supported.")
 	SIGNALPROCESSING_API void CrossCorrelate(FrequencyBuffer& FirstBufferFrequencies, FrequencyBuffer& SecondBufferFrequencies, int32 NumSamples, FrequencyBuffer& OutCorrelation);
 
-	class SIGNALPROCESSING_API FFFTConvolver
+	// Deprecated in 5.1
+	class SIGNALPROCESSING_API FFFTConvolver_DEPRECATED
 	{
 	public:
-		FFFTConvolver();
+		FFFTConvolver_DEPRECATED();
 
 		/*
 		 * Applies the convolver's internal window to InputAudio. Until SetWindow is called, ProcessAudio will not affect InputAudio.
@@ -191,6 +227,10 @@ namespace Audio
 		FAlignedFloatBuffer COLABuffer;
 	};
 
+	class UE_DEPRECATED(5.1, "Use FConvolutionFactory or UniformPartitionConvolutionFactory.") FFFTConvolver : public FFFTConvolver_DEPRECATED
+	{
+	};
+
 	// Computes the power spectrum from FFTFreqDomainData. Applies a 1/(FFTSize^2) scaling to the output to 
 	// maintain equal energy between original time domain data and output spectrum.  Only the first 
 	// (FFTSize / 2 + 1) spectrum values are calculated. These represent the frequencies from 0 to Nyquist.
@@ -198,7 +238,8 @@ namespace Audio
 	// InFrequencyData is the input frequency domain data. Generally this is created by calling PerformFFT(...)
 	// FFTSize is the number of samples used when originally calculating the FFT
 	// OutBuffer is an aligned buffer which will contain spectrum data. It will constain (FFTSize / 2 + 1) elements.
-	SIGNALPROCESSING_API void ComputePowerSpectrum(const FFTFreqDomainData& InFrequencyData, int32 FFTSize, FAlignedFloatBuffer& OutBuffer);
+	UE_DEPRECATED(5.1, "Use equivalent methods in FloatArrayMath.h")
+	SIGNALPROCESSING_API void ComputePowerSpectrum(const FFTFreqDomainData_DEPRECATED& InFrequencyData, int32 FFTSize, FAlignedFloatBuffer& OutBuffer);
 
 	// Computes the magnitude spectrum from FFTFreqDomainData. Applies a 1/FFTSize scaling to the output to 
 	// maintain equal energy between original time domain data and output spectrum.  Only the first 
@@ -207,7 +248,8 @@ namespace Audio
 	// InFrequencyData is the input frequency domain data. Generally this is created by calling PerformFFT(...)
 	// FFTSize is the number of samples used when originally calculating the FFT
 	// OutBuffer is an aligned buffer which will contain spectrum data. It will constain (FFTSize / 2 + 1) elements.
-	SIGNALPROCESSING_API void ComputeMagnitudeSpectrum(const FFTFreqDomainData& InFrequencyData, int32 FFTSize, FAlignedFloatBuffer& OutBuffer);
+	UE_DEPRECATED(5.1, "Use equivalent methods in FloatArrayMath.h")
+	SIGNALPROCESSING_API void ComputeMagnitudeSpectrum(const FFTFreqDomainData_DEPRECATED& InFrequencyData, int32 FFTSize, FAlignedFloatBuffer& OutBuffer);
 	
 	// Computes the spectrum from FFTFreqDomainData. Applies a scaling to the output to maintain equal 
 	// energy between original time domain data and output spectrum.  Only the first (FFTSize / 2 + 1)
@@ -217,5 +259,17 @@ namespace Audio
 	// InFrequencyData is the input frequency domain data. Generally this is created by calling PerformFFT(...)
 	// FFTSize is the number of samples used when originally calculating the FFT
 	// OutBuffer is an aligned buffer which will contain spectrum data. It will constain (FFTSize / 2 + 1) elements.
-	SIGNALPROCESSING_API void ComputeSpectrum(ESpectrumType InSpectrumType, const FFTFreqDomainData& InFrequencyData, int32 FFTSize, FAlignedFloatBuffer& OutBuffer);
+	UE_DEPRECATED(5.1, "Use equivalent methods in FloatArrayMath.h")
+	SIGNALPROCESSING_API void ComputeSpectrum(ESpectrumType InSpectrumType, const FFTFreqDomainData_DEPRECATED& InFrequencyData, int32 FFTSize, FAlignedFloatBuffer& OutBuffer);
+
+	// Return the ceiling of the log2 of InNum
+	SIGNALPROCESSING_API int32 CeilLog2(int32 InNum);
+
+	// Return the scaling factor needed to apply to a power spectrum given a current
+	// and target FFT scaling. 
+	SIGNALPROCESSING_API float GetPowerSpectrumScaling(int32 FFTSize, EFFTScaling InCurrentScaling, EFFTScaling InTargetScaling);
+
+	// Scale the power spectrum to remove any scaling introduced by the FFT algorithm
+	// implementation.
+	SIGNALPROCESSING_API void ScalePowerSpectrumInPlace(int32 FFTSize, EFFTScaling InCurrentScaling, EFFTScaling InTargetScaling, TArrayView<float> InPowerSpectrum);
 }
diff --git a/Engine/Source/Runtime/SignalProcessing/Public/DSP/PassiveFilter.h b/Engine/Source/Runtime/SignalProcessing/Public/DSP/PassiveFilter.h
--- a/Engine/Source/Runtime/SignalProcessing/Public/DSP/PassiveFilter.h
+++ b/Engine/Source/Runtime/SignalProcessing/Public/DSP/PassiveFilter.h
@@ -1,9 +1,10 @@
 // Copyright Epic Games, Inc. All Rights Reserved.
 
 #pragma once
-
+ 
 #include "CoreMinimal.h"
 #include "DSP/AudioFFT.h"
+#include "DSP/FFTAlgorithm.h"
 #include "Async/ParallelFor.h"
 #include "Audio.h"
 
@@ -46,9 +47,6 @@ namespace Audio
 		}
 	};
 
-	UE_DEPRECATED(4.23, "PassiveFilterParams has been renamed to FPassiveFilterParams")
-	typedef FPassiveFilterParams PassiveFilterParams;
-
 	static float EvaluateChebyshevPolynomial(float FrequencyRatio, int32 Order)
 	{
 		if (FMath::IsNearlyEqual(FrequencyRatio, 1.0f))
@@ -127,19 +125,16 @@ namespace Audio
 			return;
 		}
 
-		FAlignedFloatBuffer TempReal;
-		FAlignedFloatBuffer TempImag;
-		TempReal.AddUninitialized(Signal.Num());
-		TempImag.AddUninitialized(Signal.Num());
-
-		// Perform FFT on data:
-		FFTTimeDomainData TimeData;
-		TimeData.Buffer = Signal.GetData();
-		TimeData.NumSamples = Signal.Num();
+		FFFTSettings FFTSettings;
+		FFTSettings.Log2Size = CeilLog2(Signal.Num());
+		FFTSettings.bArrays128BitAligned = true;
 
-		FFTFreqDomainData FreqData;
-		FreqData.OutReal = TempReal.GetData();
-		FreqData.OutImag = TempImag.GetData();
+		TUniquePtr<IFFTAlgorithm> FFT = FFFTFactory::NewFFTAlgorithm(FFTSettings);
+		if (!FFT.IsValid())
+		{
+			UE_LOG(LogAudio, Error, TEXT("Failed to create an FFT Algorithm with log2size %n"), FFTSettings.Log2Size);
+			return;
+		}
 
 		float MaxVal = Signal[0];
 		float MinVal = Signal[0];
@@ -171,32 +166,38 @@ namespace Audio
 			}
 		}
 
-		PerformFFT(TimeData, FreqData);
+		// Perform FFT on data:
+		FAlignedFloatBuffer TimeData;
+		TimeData.AddZeroed(FFT->NumInputFloats());
+		FMemory::Memcpy(TimeData.GetData(), Signal.GetData(), Signal.Num() * sizeof(float));
+
+		FAlignedFloatBuffer ComplexSpectrum;
+		ComplexSpectrum.AddUninitialized(FFT->NumOutputFloats());
 
-		const int32 NumSamples = Signal.Num();
-		const float NumBins = NumSamples / 2;
+		FFT->ForwardRealToComplex(TimeData.GetData(), ComplexSpectrum.GetData());
+
+		const float NumBins = FFT->NumOutputFloats() / 2;
 
 		if (InParams.bRemoveDC)
 		{
-			FreqData.OutReal[0] = 0.0f;
-			FreqData.OutImag[0] = 0.0f;
+			ComplexSpectrum[0] = 0.f;
+			ComplexSpectrum[1] = 0.f;
 		}
 
 		// Apply filter in parallel:
 		ParallelFor(NumBins, [&](int32 Index)
 		{
-			float NormalizedFreq = ((float)Index) / NumBins;
+			float NormalizedFreq = ((float)Index) / (NumBins - 1);
 			float Gain = GetGainForFrequency(NormalizedFreq, InParams);
 
-			FreqData.OutReal[Index] *= Gain;
-			FreqData.OutImag[Index] *= Gain;
-
-			FreqData.OutReal[NumSamples - Index - 1] *= Gain;
-			FreqData.OutImag[NumSamples - Index - 1] *= Gain;
+			ComplexSpectrum[2 * Index] *= Gain;
+			ComplexSpectrum[2 * Index + 1] *= Gain;
 		});
 
 		// Inverse FFT back into the signal:
-		PerformIFFT(FreqData, TimeData);
+		FFT->InverseComplexToReal(ComplexSpectrum.GetData(), TimeData.GetData());
+
+		FMemory::Memcpy(Signal.GetData(), TimeData.GetData(), Signal.Num() * sizeof(float));
 
 		// If we're scaling, map back to it's original range:
 		if (InParams.bScaleByOffset)
diff --git a/Engine/Source/Runtime/SignalProcessing/Public/DSP/VoiceProcessing.h b/Engine/Source/Runtime/SignalProcessing/Public/DSP/VoiceProcessing.h
--- a/Engine/Source/Runtime/SignalProcessing/Public/DSP/VoiceProcessing.h
+++ b/Engine/Source/Runtime/SignalProcessing/Public/DSP/VoiceProcessing.h
@@ -128,16 +128,22 @@ namespace Audio
 		float GainMax;
 	};
 	
+	/** VoiceProcessing.h Deprecation
+	 *
+	 * Several classes in VoiceProcessing are deprecated due to lack of support
+	 * and lack of need. Current voice processing solutions are available through 
+	 * EOS and WebRTC. 
+	 */
 	
 	/**
 	 * This filter takes a precomputed set of FIR weights in the frequency domain, and linearly converges to it.
 	 * If no new weights are given and we've converged to our previous input weights, this works like a normal FFT-based FIR filter.
 	 * Convergence is non-asymptotic: if no new weights are given after our number of steps until convergence, our filter is using the exact weights given.
 	 */
-	class SIGNALPROCESSING_API FAdaptiveFilter
+	class SIGNALPROCESSING_API FAdaptiveFilter_DEPRECATED // Deprecated in 5.1
 	{
 	public:
-		FAdaptiveFilter(int32 FilterLength, int32 AudioCallbackSize);
+		FAdaptiveFilter_DEPRECATED(int32 FilterLength, int32 AudioCallbackSize);
 
 		/*
 		* Applies current filter to InAudio in-place. If there is a new target set of weights, they can be input below.
@@ -148,7 +154,7 @@ namespace Audio
 
 	private:
 
-		FAdaptiveFilter();
+		FAdaptiveFilter_DEPRECATED();
 
 		/** This is called every process callback to (A) reset our weight deltas if we have a new target, and (B) increment our current filter weights if we haven't converged yet. */
 		void AdaptFilter();
@@ -168,16 +174,23 @@ namespace Audio
 
 		int32 CurrentStepsUntilConvergence;
 
-		FFFTConvolver Convolver;
+		FFFTConvolver_DEPRECATED Convolver;
+	};
+
+	class UE_DEPRECATED(5.1, "FAdaptiveFilter will no longer be supported.") FAdaptiveFilter;
+	class SIGNALPROCESSING_API FAdaptiveFilter : public FAdaptiveFilter_DEPRECATED
+	{
+		public:
+			using FAdaptiveFilter_DEPRECATED::FAdaptiveFilter_DEPRECATED;
 	};
 
 	/**
 	 * This class takes an incoming signal and an outgoing signal, Correlates them, and returns the frequency values of the weight targets to pass to an adaptive filter.  
 	 */
-	class SIGNALPROCESSING_API FFDAPFilterComputer
+	class SIGNALPROCESSING_API FFDAPFilterComputer_DEPRECATED // Deprecated in 5.1
 	{
 	public:
-		FFDAPFilterComputer();
+		FFDAPFilterComputer_DEPRECATED();
 
 		void GenerateWeights(const float* IncomingSignal, int32 NumIncomingSamples, const float* OutgoingSignal, int32 NumOutgoingSamples, FrequencyBuffer& OutWeights);
 
@@ -189,18 +202,23 @@ namespace Audio
 		FAlignedFloatBuffer ZeroPaddedOutgoingBuffer;
 	};
 
+	class UE_DEPRECATED(5.1, "FAFDAPFFilterComputer will no longer be supported.") FFDAPFilterComputer;
+	class SIGNALPROCESSING_API FFDAPFilterComputer : public FFDAPFilterComputer_DEPRECATED 
+	{
+	};
+
 	/*
 	 * This class uses an adaptive filter to cancel out any rendered audio signal that might be picked up by the mic.
 	 * To add a new patch to a rendered audio signal, user AddNewSignalPatch. See FPatchInput for how to push audio.
 	 * ProcessAudio then filters the microphone signal accordingly.
 	 */
-	class SIGNALPROCESSING_API FAcousticEchoCancellation
+	class SIGNALPROCESSING_API FAcousticEchoCancellation_DEPRECATED // Deprecated in 5.1
 	{
 	public:
 		/**
 		 * Convergence Rate should be a number between 0 and 1. The higher the number, the quicker the adaptive filter reacts. 
 		 */
-		FAcousticEchoCancellation(float InConvergenceRate, int32 CallbackSize, int32 InFilterLength, int32 InFilterUpdateRate = 1);
+		FAcousticEchoCancellation_DEPRECATED(float InConvergenceRate, int32 CallbackSize, int32 InFilterLength, int32 InFilterUpdateRate = 1);
 
 		/** Callback function for outgoing audio signal. This is where the filter is applied, and the bulk of the DSP work takes place. */
 		void ProcessAudio(float* InAudio, int32 NumSamples);
@@ -211,8 +229,8 @@ namespace Audio
 
 	private:
 		FPatchMixer PatchMixer;
-		FFDAPFilterComputer FilterComputer;
-		FAdaptiveFilter AdaptiveFilter;
+		FFDAPFilterComputer_DEPRECATED FilterComputer;
+		FAdaptiveFilter_DEPRECATED AdaptiveFilter;
 		
 		FAlignedFloatBuffer FilterComputerInput;
 		FrequencyBuffer FilterComputerOutput;
@@ -221,4 +239,11 @@ namespace Audio
 		int32 FilterUpdateRate;
 		int32 FitlerUpdateCounter;
 	};
+
+	class UE_DEPRECATED(5.1, "FAcousticEchoCancellation will no longer be supported.") FAcousticEchoCancellation;
+	class SIGNALPROCESSING_API FAcousticEchoCancellation : public FAcousticEchoCancellation_DEPRECATED
+	{
+		public:
+			using FAcousticEchoCancellation_DEPRECATED::FAcousticEchoCancellation_DEPRECATED;
+	};
 }
