CIS Fix for Mac UnitTestAudioFFT aggregate initialization

diff --git a/Engine/Source/Runtime/SignalProcessing/Public/DSP/AudioFFT.h b/Engine/Source/Runtime/SignalProcessing/Public/DSP/AudioFFT.h
--- a/Engine/Source/Runtime/SignalProcessing/Public/DSP/AudioFFT.h
+++ b/Engine/Source/Runtime/SignalProcessing/Public/DSP/AudioFFT.h
@@ -95,6 +95,13 @@ namespace Audio
 	struct UE_DEPRECATED(5.1, "Use IFFTAlgorithm and TArray for time domain data") FFTTimeDomainData;
 	struct FFTTimeDomainData : FFTTimeDomainData_DEPRECATED 
 	{
+		FFTTimeDomainData() = default;
+
+		FFTTimeDomainData(float* InBuffer, int32 InNumSamples)
+		{
+			Buffer = InBuffer;
+			NumSamples = InNumSamples;
+		}
 	};
 
 	struct FFTFreqDomainData_DEPRECATED
@@ -107,6 +114,13 @@ namespace Audio
 	struct UE_DEPRECATED(5.1, "Use IFFTAlgorithm and TArray for time freq domain data") FFTFreqDomainData;
 	struct FFTFreqDomainData : FFTFreqDomainData_DEPRECATED 
 	{
+		FFTFreqDomainData() = default;
+
+		FFTFreqDomainData(float* Real, float* Imag)
+		{
+			OutReal = Real;
+			OutImag = Imag;
+		}
 	};
 
 	// Performs a one-time FFT on a float buffer. Does not support complex signals.
