Fix crash in submix spectrum analyzer with unclamped hop size

diff --git a/Engine/Source/Runtime/Engine/Private/SoundSubmix.cpp b/Engine/Source/Runtime/Engine/Private/SoundSubmix.cpp
--- a/Engine/Source/Runtime/Engine/Private/SoundSubmix.cpp
+++ b/Engine/Source/Runtime/Engine/Private/SoundSubmix.cpp
@@ -970,9 +970,12 @@ FSoundSpectrumAnalyzerSettings USoundSubmix::GetSpectrumAnalyzerSettings(EFFTSiz
 	OutSettings.FFTSize = FFTSize;
 	OutSettings.WindowType = WindowType;
 	OutSettings.InterpolationMethod = InterpolationMethod;
-	OutSettings.HopSize = HopSize;
 	OutSettings.SpectrumType = SpectrumType;
 
+	const float MinHopSize = 0.001f;
+	const float MaxHopSize = 10.0f;
+	OutSettings.HopSize = FMath::Clamp(HopSize, MinHopSize, MaxHopSize);
+
 	return OutSettings;
 }
 
diff --git a/Engine/Source/Runtime/SignalProcessing/Private/SpectrumAnalyzer.cpp b/Engine/Source/Runtime/SignalProcessing/Private/SpectrumAnalyzer.cpp
--- a/Engine/Source/Runtime/SignalProcessing/Private/SpectrumAnalyzer.cpp
+++ b/Engine/Source/Runtime/SignalProcessing/Private/SpectrumAnalyzer.cpp
@@ -1058,7 +1058,8 @@ namespace Audio
 
 		// If we have enough audio pushed to the spectrum analyzer and we have an available buffer to work in,
 		// we can start analyzing.
-		if (InputQueue.Num() >= ((uint32)FFTSize))
+		uint32 RequiredSize = FMath::Max(FFTSize, HopInSamples);
+		if (InputQueue.Num() >= RequiredSize)
 		{
 			int64 WindowSampleCenterIndex = 0;
 
@@ -1066,16 +1067,22 @@ namespace Audio
 
 			if (bUseLatestAudio)
 			{
+				WindowSampleCenterIndex = SampleCounter.GetValue() - (FFTSize / 2);
+
 				// If we are only using the latest audio, scrap the oldest audio in the InputQueue:
 				InputQueue.SetNum((uint32)FFTSize);
+				InputQueue.Pop(TimeDomainBuffer, FFTSize);
+			}
+			else
+			{
+				WindowSampleCenterIndex = SampleCounter.GetValue() - InputQueue.Num() + (FFTSize / 2);
+				
+				// Perform pop/peek here based on FFT size and hop amount.
+				InputQueue.Peek(TimeDomainBuffer, FFTSize);
+				InputQueue.Pop(HopInSamples);
 			}
-			WindowSampleCenterIndex = SampleCounter.GetValue() - InputQueue.Num() + FFTSize / 2;
-			double Timestamp = static_cast<double>(WindowSampleCenterIndex) / FMath::Max(SampleRate, 1.f);
 
-			// Perform pop/peek here based on FFT size and hop amount.
-			const int32 PeekAmount = FFTSize - HopInSamples;
-			InputQueue.Pop(TimeDomainBuffer, HopInSamples);
-			InputQueue.Peek(TimeDomainBuffer + HopInSamples, PeekAmount);
+			double Timestamp = static_cast<double>(WindowSampleCenterIndex) / FMath::Max(SampleRate, 1.f);
 
 			// apply window if necessary.
 			Window.ApplyToBuffer(TimeDomainBuffer);
