DSP LFO/Modulation LFO fixes
- Add missing exponential factor to Modulation generator
- Fix bounds restriction calculation for expontential factor
- Add LFO ResetPhase implementation to override initial location in oscilator on rising edge where necessary (for modulation LFO, user range limited to unipolar but hardly ever wants to start at midpoint).
- Add bipolar option to avoid oneshot/bipolar conflation for certain types of LFOs
- Add back implementation for 5.1 deprecated function to avoid breaking licensees

diff --git a/Engine/Plugins/Runtime/AudioModulation/Source/AudioModulation/Private/Generators/SoundModulationLFO.cpp b/Engine/Plugins/Runtime/AudioModulation/Source/AudioModulation/Private/Generators/SoundModulationLFO.cpp
--- a/Engine/Plugins/Runtime/AudioModulation/Source/AudioModulation/Private/Generators/SoundModulationLFO.cpp
+++ b/Engine/Plugins/Runtime/AudioModulation/Source/AudioModulation/Private/Generators/SoundModulationLFO.cpp
@@ -41,13 +41,16 @@ namespace AudioModulation
 			FLFOGenerator(const FSoundModulationLFOParams& InParams)
 				: Params(InParams)
 			{
-				LFO.SetGain(Params.Amplitude);
+				LFO.SetBipolar(false);
+				LFO.SetExponentialFactor(Params.ExponentialFactor);
 				LFO.SetFrequency(Params.Frequency);
+				LFO.SetGain(Params.Amplitude);
 				LFO.SetMode(Params.bLooping ? Audio::ELFOMode::Type::Sync : Audio::ELFOMode::OneShot);
-
-				static_assert(static_cast<int32>(ESoundModulationLFOShape::COUNT) == static_cast<int32>(Audio::ELFO::Type::NumLFOTypes), "LFOShape/ELFO Type mismatch");
+				LFO.SetPulseWidth(Params.Width);
 				LFO.SetType(static_cast<Audio::ELFO::Type>(Params.Shape));
 				LFO.Start();
+
+				static_assert(static_cast<int32>(ESoundModulationLFOShape::COUNT) == static_cast<int32>(Audio::ELFO::Type::NumLFOTypes), "LFOShape/ELFO Type mismatch");
 			}
 
 			virtual ~FLFOGenerator() = default;
@@ -90,8 +93,11 @@ namespace AudioModulation
 					const FLFOGenerator* Generator = static_cast<const FLFOGenerator*>(NewGenerator.Get());
 					Params = Generator->Params;
 
-					LFO.SetGain(Params.Amplitude);
+					LFO.SetExponentialFactor(Params.ExponentialFactor);
 					LFO.SetFrequency(Params.Frequency);
+					LFO.SetGain(Params.Amplitude);
+					LFO.SetPhaseOffset(Params.Phase);
+					LFO.SetPulseWidth(Params.Width);
 
 					static_assert(static_cast<int32>(ESoundModulationLFOShape::COUNT) == static_cast<int32>(Audio::ELFO::Type::NumLFOTypes), "LFOShape/ELFO Type mismatch");
 					LFO.SetType(static_cast<Audio::ELFO::Type>(Params.Shape));
diff --git a/Engine/Plugins/Runtime/AudioModulation/Source/AudioModulation/Public/Generators/SoundModulationLFO.h b/Engine/Plugins/Runtime/AudioModulation/Source/AudioModulation/Public/Generators/SoundModulationLFO.h
--- a/Engine/Plugins/Runtime/AudioModulation/Source/AudioModulation/Public/Generators/SoundModulationLFO.h
+++ b/Engine/Plugins/Runtime/AudioModulation/Source/AudioModulation/Public/Generators/SoundModulationLFO.h
@@ -35,25 +35,37 @@ struct FSoundModulationLFOParams
 	UPROPERTY(EditAnywhere, Category = Modulation, BlueprintReadWrite, meta = (DisplayPriority = 20, EditCondition = "!bBypass"))
 	ESoundModulationLFOShape Shape = ESoundModulationLFOShape::Sine;
 
+	/** Factor increasing/decreasing curvature of exponential LFO shape. */
+	UPROPERTY(EditAnywhere, Category = Modulation, BlueprintReadWrite, meta = (DisplayPriority = 21, EditCondition = "!bBypass && Shape == ESoundModulationLFOShape::Exponential", EditConditionHides, ClampMin = "0.000001", UIMin = "0.25", UIMax = "10.0"))
+	float ExponentialFactor = 3.5f;
+
+	/** Pulse width of square LFO shape. */
+	UPROPERTY(EditAnywhere, Category = Modulation, BlueprintReadWrite, meta = (DisplayPriority = 21, EditCondition = "!bBypass && Shape == ESoundModulationLFOShape::Square", EditConditionHides, ClampMin = "0.0", ClampMax = "1.0"))
+	float Width = 0.5f;
+
 	/** Amplitude of oscillator */
 	UPROPERTY(EditAnywhere, Category = Modulation, BlueprintReadWrite, meta = (DisplayPriority = 30, EditCondition = "!bBypass", UIMin = "0", UIMax = "1", ClampMin = "0", ClampMax = "1"))
-	float Amplitude = 0.5f;
+	float Amplitude = 1.0f;
 
 	/** Frequency of oscillator */
 	UPROPERTY(EditAnywhere, Category = Modulation, BlueprintReadWrite, meta = (DisplayPriority = 40, EditCondition = "!bBypass", UIMin = "0", UIMax = "20", ClampMin = "0", ClampMax = "20"))
 	float Frequency = 1.0f;
 
-	/** Offset of oscillator */
+	/** Amplitude offset of oscillator */
 	UPROPERTY(EditAnywhere, Category = Modulation, BlueprintReadWrite, meta = (DisplayPriority = 50, EditCondition = "!bBypass", UIMin = "0", UIMax = "1", ClampMin = "0", ClampMax = "1"))
-	float Offset = 0.5;
+	float Offset = 0.0f;
+
+	/** Unit phase offset of oscillator */
+	UPROPERTY(EditAnywhere, Category = Modulation, BlueprintReadWrite, meta = (DisplayPriority = 60, EditCondition = "!bBypass", UIMin = "0", UIMax = "1", ClampMin = "0"))
+	float Phase = 0.0f;
 
 	/** Whether or not to loop the oscillation more than once */
-	UPROPERTY(EditAnywhere, Category = Modulation, BlueprintReadWrite, meta = (DisplayPriority = 60, EditCondition = "!bBypass"))
-	bool bLooping = 1;
+	UPROPERTY(EditAnywhere, Category = Modulation, BlueprintReadWrite, meta = (DisplayPriority = 70, EditCondition = "!bBypass"))
+	bool bLooping = true;
 
 	/** If true, bypasses LFO bus from being modulated by parameters, patches, or mixed (LFO remains active and computed). */
 	UPROPERTY(EditAnywhere, Category = Modulation, BlueprintReadWrite, meta = (DisplayPriority = 10))
-	bool bBypass = 0;
+	bool bBypass = false;
 };
 
 
diff --git a/Engine/Plugins/Runtime/Synthesis/Source/Synthesis/Private/EpicSynth1.cpp b/Engine/Plugins/Runtime/Synthesis/Source/Synthesis/Private/EpicSynth1.cpp
--- a/Engine/Plugins/Runtime/Synthesis/Source/Synthesis/Private/EpicSynth1.cpp
+++ b/Engine/Plugins/Runtime/Synthesis/Source/Synthesis/Private/EpicSynth1.cpp
@@ -49,6 +49,7 @@ namespace Audio
 		for (int32 i = 0; i < NumLFOs; ++i)
 		{
 			LFO[i].Init(ControlSampleRate, VoiceId, ModMatrix);
+			LFO[i].SetBipolar(false);
 		}
 		GainEnv.Init(ControlSampleRate, VoiceId, ModMatrix);
 		ModEnv.Init(ControlSampleRate, VoiceId, ModMatrix);
diff --git a/Engine/Plugins/Runtime/Synthesis/Source/Synthesis/Private/MonoWaveTable.cpp b/Engine/Plugins/Runtime/Synthesis/Source/Synthesis/Private/MonoWaveTable.cpp
--- a/Engine/Plugins/Runtime/Synthesis/Source/Synthesis/Private/MonoWaveTable.cpp
+++ b/Engine/Plugins/Runtime/Synthesis/Source/Synthesis/Private/MonoWaveTable.cpp
@@ -33,30 +33,31 @@ namespace Audio
 
 			// LFO
 			PosLfo.Init(OutputSampleRate);
-			PosLfo.Start();
+			PosLfo.SetBipolar(false);
 			PosLfo.SetType(Audio::ELFO::Sine);
-			PosLfo.SetFrequency(0.0);
-			PosLfo.SetGainMod(0.0);
+			PosLfo.SetFrequency(0.0f);
+			PosLfo.SetGainMod(0.0f);
+			PosLfo.Start();
 
 			// Envelopes
 			AmpEnv.SetAttackTime(1);
-			AmpEnv.SetSustainGain(1.0);
+			AmpEnv.SetSustainGain(1.0f);
 			AmpEnv.SetReleaseTime(4000);
-			AmpEnv.SetDepth(1.0);	
+			AmpEnv.SetDepth(1.0f);	
 			AmpEnv.SetLegato(true);
 			AmpEnv.SetRetrigger(true);
 
 			LPFEnv.SetAttackTime(1);
-			LPFEnv.SetSustainGain(1.0);
+			LPFEnv.SetSustainGain(1.0f);
 			LPFEnv.SetReleaseTime(4000);
-			LPFEnv.SetDepth(1.0);
+			LPFEnv.SetDepth(1.0f);
 			LPFEnv.SetLegato(true);
 			LPFEnv.SetRetrigger(true);
 
 			// Filters
 			LPFBaseFreqHz = 1000;
 			LPFMaxModOffsetOctaves = 9;
-			LPFCeilingHz = 8'000;
+			LPFCeilingHz = 8000;
 			LastLPFFreqHz = LPFBaseFreqHz;
 			LPF.Init(OutputSampleRate, 1);
 			LPF.SetFrequency(LPFBaseFreqHz);
diff --git a/Engine/Source/Runtime/SignalProcessing/Private/LFO.cpp b/Engine/Source/Runtime/SignalProcessing/Private/LFO.cpp
--- a/Engine/Source/Runtime/SignalProcessing/Private/LFO.cpp
+++ b/Engine/Source/Runtime/SignalProcessing/Private/LFO.cpp
@@ -3,6 +3,7 @@
 #include "DSP/LFO.h"
 #include "DSP/Dsp.h"
 
+
 namespace Audio
 {
 	FLFO::FLFO()
@@ -14,11 +15,9 @@ namespace Audio
 		, ModScale(1.0f)
 		, ModAdd(0.0f)
 		, LastOutput(0.0f)
+		, LoopCount(0.0f)
 		, QuadLastOutput(0.0f)
-	{
-	}
-
-	FLFO::~FLFO()
+		, bBipolar(true)
 	{
 	}
 
@@ -44,6 +43,10 @@ namespace Audio
 		{
 			Reset();
 		}
+		else if (!bIsPlaying)
+		{
+			ResetPhase();
+		}
 
 		bIsPlaying = true;
 	}
@@ -58,31 +61,103 @@ namespace Audio
 		// reset base class first
 		IOscBase::Reset();
 
+		ResetPhase();
+
 		RSHValue = 0.0f;
 		RSHCounter = INDEX_NONE;
 	}
-	
-	float FLFO::Generate(float* QuadPhaseOutput)
+
+	void FLFO::ResetPhase()
 	{
-		// If the LFO isn't playing, return 0.0 for quad phase and for normal output
-		if (!bIsPlaying)
+		// Reset loop count
+		LoopCount = 0.0f;
+
+		// Set initial default phase to zero crossing and rising edge where
+		// possible (omits "phase offset" input to allow for client system
+		// to set to non-zero if desired).
+		switch (LFOType)
 		{
-			if (QuadPhaseOutput)
+			case ELFO::Sine:
+			case ELFO::Triangle:
 			{
-				*QuadPhaseOutput = QuadLastOutput;
+				Phase = bBipolar ? 0.25f : 0.0f;
 			}
+			break;
 
-			return LastOutput;
-		}
+			case ELFO::DownSaw:
+			case ELFO::UpSaw:
+			{
+				Phase = bBipolar ? 0.5f : 0.0f;
+			}
+			break;
 
-		const bool bWrapped = WrapPhase();
+			case ELFO::Exponential:
+			{
+				Phase = bBipolar ? FMath::Pow(0.5f, 1.0f / ExponentialFactor) : 0.0f;
+			}
+			break;
 
-		// If we're in the oneshot mode, check if we've wrapped
-		// and if so, turn the LFO off and return 0.0s
-		if (LFOMode == ELFOMode::OneShot && bWrapped)
-		{
-			bIsPlaying = false;
+			case ELFO::RandomSampleHold:
+			case ELFO::Square:
+			default:
+			{
+				static_assert(static_cast<int32>(ELFO::NumLFOTypes) == 7, "Possible missing switch case coverage");
+				Phase = 0.0f;
+			}
+			break;
+		};
+	}
+
+	void FLFO::SetBipolar(const bool bInBipolar)
+	{
+		bBipolar = bInBipolar;
+	}
+
+	void FLFO::SetPhaseOffset(const float InOffset)
+	{
+		PhaseOffset = FMath::Fmod(FMath::Max(0.0f, InOffset), 1.0f);
+	}
 
+	void FLFO::SetType(const ELFO::Type InLFOType)
+	{
+		LFOType = InLFOType;
+	}
+
+	ELFO::Type FLFO::GetType() const
+	{
+		return LFOType;
+	}
+
+	void FLFO::SetMode(const ELFOMode::Type InLFOMode)
+	{
+		LFOMode = InLFOMode;
+	}
+
+	ELFOMode::Type FLFO::GetMode() const
+	{
+		return LFOMode;
+	}
+
+	void FLFO::SetExponentialFactor(const float InExpFactor)
+	{
+		ExponentialFactor = FMath::Max(InExpFactor, UE_SMALL_NUMBER);
+	}
+
+	FPatchSource FLFO::GetModSourceNormalPhase() const
+	{
+		return ModNormalPhase;
+	}
+
+	FPatchSource FLFO::GetModSourceQuadPhase() const
+	{
+		return ModQuadPhase;
+	}
+
+	float FLFO::Generate(float* QuadPhaseOutput)
+	{
+		// If the LFO isn't playing, return last computed value for both output & quad.
+		if (!bIsPlaying)
+		{
 			if (QuadPhaseOutput)
 			{
 				*QuadPhaseOutput = QuadLastOutput;
@@ -91,11 +166,23 @@ namespace Audio
 			return LastOutput;
 		}
 
+		WrapPhase();
 		LastOutput = ComputeLFO(GetPhase(), QuadPhaseOutput);
 
 		// Update the LFO phase after computing LFO values
+		LoopCount += PhaseInc;
 		UpdatePhase();
 
+		// If in oneshot mode, check if wrapped and if so, turn the LFO off and compute last value.
+		if (LFOMode == ELFOMode::OneShot)
+		{
+			if (LoopCount >= 1.0f)
+			{
+				bIsPlaying = false;
+				LoopCount = 0.0f;
+			}
+		}
+
 		// Return the output
 		return LastOutput;
 	}
@@ -105,62 +192,60 @@ namespace Audio
 		float Output = 0.0f;
 		float QuadOutput = 0.0f;
 
-		float QuadPhase = InPhase + 0.25f;
-		if (QuadPhase >= 1.0f)
-		{
-			QuadPhase -= 1.0f;
-		}
+		const float CurPhase = FMath::Fmod(InPhase + PhaseOffset, 1.0f);
+		const float QuadPhase = FMath::Fmod(InPhase + PhaseOffset + 0.25f , 1.0f);
 
 		switch (LFOType)
 		{
 			case ELFO::Sine:
 			{
-				float Angle = 2.0f * InPhase * PI - PI;
-				Output = Audio::FastSin(Angle);
-
-				Angle = 2.0f * QuadPhase * PI - PI;
-				QuadOutput = Audio::FastSin(Angle);
+				// Must subtract pi and flip sign to guarantee in valid range for FastSin function,
+				// yet still starts on rising edge from 0 crossing.
+				auto ComputeSine = [](float InputPhase)
+				{
+					if (InputPhase > 0.5f)
+					{
+						InputPhase -= 1.0f;
+					}
+					const float Angle = 2.0f * InputPhase * PI;
+					return 0.5f * Audio::FastSin(Angle) + 0.5f;
+				};
+				Output = ComputeSine(CurPhase);
+				QuadOutput = ComputeSine(QuadPhase);
 			}
 			break;
 
 			case ELFO::UpSaw:
 			{
-				Output = GetBipolar(InPhase);
-				QuadOutput = GetBipolar(QuadPhase);
+				Output = CurPhase;
+				QuadOutput = QuadPhase;
 			}
 			break;
 
 			case ELFO::DownSaw:
 			{
-				Output = -1.0f * GetBipolar(InPhase);
-				QuadOutput = -1.0f * GetBipolar(QuadPhase);
+				Output = 1.0f - CurPhase;
+				QuadOutput = 1.0f - QuadPhase;
 			}
 			break;
 
 			case ELFO::Square:
 			{
-				Output = InPhase > PulseWidth ? -1.0f : 1.0f;
-				QuadOutput = QuadPhase > PulseWidth ? -1.0f : 1.0f;
+				Output = CurPhase > PulseWidth ? 0.0f : 1.0f;
+				QuadOutput = QuadPhase > PulseWidth ? 0.0f : 1.0f;
 			}
 			break;
 
 			case ELFO::Triangle:
 			{
-				Output = FMath::Abs(GetBipolar(InPhase));
-				QuadOutput = FMath::Abs(GetBipolar(QuadPhase));;
-
-				// If not one-shot, we need to convert to bipolar
-				if (LFOMode != ELFOMode::OneShot)
-				{
-					Output = GetBipolar(Output);
-					QuadOutput = GetBipolar(QuadOutput);
-				}
+				Output = 1.0f - FMath::Abs(GetBipolar(CurPhase));
+				QuadOutput = 1.0f - FMath::Abs(GetBipolar(QuadPhase));
 			}
 			break;
 
 			case ELFO::Exponential:
 			{
-				Output = FMath::Pow(InPhase, ExponentialFactor);
+				Output = FMath::Pow(CurPhase, ExponentialFactor);
 				QuadOutput = FMath::Pow(QuadPhase, ExponentialFactor);
 			}
 			break;
@@ -171,7 +256,7 @@ namespace Audio
 				if (RSHCounter > (uint32)FrequencyThreshold)
 				{
 					RSHCounter = 0;
-					RSHValue = FMath::FRandRange(-1.0f, 1.0f);
+					RSHValue = FMath::FRand();
 				}
 				else
 				{
@@ -184,6 +269,12 @@ namespace Audio
 			break;
 		}
 
+		if (bBipolar)
+		{
+			Output = GetBipolar(Output);
+			QuadOutput = GetBipolar(QuadOutput);
+		}
+
 		const float MaxGain = Gain * ExternalGainMod;
 		Output = Output * MaxGain;
 		QuadOutput = QuadOutput * MaxGain;
@@ -212,4 +303,4 @@ namespace Audio
 
 		return Output;
 	}
-}
+} // namespace Audio
diff --git a/Engine/Source/Runtime/SignalProcessing/Private/Osc.cpp b/Engine/Source/Runtime/SignalProcessing/Private/Osc.cpp
--- a/Engine/Source/Runtime/SignalProcessing/Private/Osc.cpp
+++ b/Engine/Source/Runtime/SignalProcessing/Private/Osc.cpp
@@ -174,7 +174,11 @@ namespace Audio
 
 	void IOscBase::SetPulseWidth(const float InPulseWidth)
 	{
-		PulseWidthBase = FMath::Clamp(InPulseWidth, 0.0f, 1.0f);
+		if (InPulseWidth != PulseWidthBase)
+		{
+			PulseWidthBase = FMath::Clamp(InPulseWidth, 0.0f, 1.0f);
+			bChanged = true;
+		}
 	}
 
 	void IOscBase::ResetPhase()
@@ -182,6 +186,11 @@ namespace Audio
 		Phase = 0.0f;
 	}
 
+	void IOscBase::SetSlaveOsc(IOscBase* InSlaveOsc)
+	{
+		FollowerOsc = InSlaveOsc;
+	}
+
 	void IOscBase::SetFollowerOsc(IOscBase* InFollowerOsc)
 	{
 		FollowerOsc = InFollowerOsc;
diff --git a/Engine/Source/Runtime/SignalProcessing/Public/DSP/LFO.h b/Engine/Source/Runtime/SignalProcessing/Public/DSP/LFO.h
--- a/Engine/Source/Runtime/SignalProcessing/Public/DSP/LFO.h
+++ b/Engine/Source/Runtime/SignalProcessing/Public/DSP/LFO.h
@@ -1,11 +1,10 @@
 // Copyright Epic Games, Inc. All Rights Reserved.
-
 #pragma once
 
-#include "CoreMinimal.h"
 #include "DSP/Osc.h"
 #include "DSP/ModulationMatrix.h"
 
+
 namespace Audio
 {
 	namespace ELFO
@@ -46,7 +45,7 @@ namespace Audio
 	{
 	public:
 		FLFO();
-		virtual ~FLFO();
+		virtual ~FLFO() = default;
 
 		//~ Begin FOscBase
 		virtual void Init(const float InSampleRate, const int32 InVoiceId = 0, FModulationMatrix* InMatrix = nullptr, const int32 ModMatrixStage = 0) override;
@@ -56,17 +55,43 @@ namespace Audio
 		virtual float Generate(float* QuadPhaseOutput = nullptr) override;
 		//~ End FOscBase
 
-		void SetType(const ELFO::Type InLFOType) { LFOType = InLFOType; }
-		ELFO::Type GetType() const { return LFOType; }
-		void SetMode(const ELFOMode::Type InLFOMode) { LFOMode = InLFOMode; }
-		ELFOMode::Type GetMode() const { return LFOMode; }
-		void SetExponentialFactor(const float InExpFactor) { ExponentialFactor = FMath::Min(InExpFactor, 1.0f); }
-		FPatchSource GetModSourceNormalPhase() const { return ModNormalPhase; }
-		FPatchSource GetModSourceQuadPhase() const { return ModQuadPhase; }
+		// Set whether or not LFO outputs in bipolar domain or unipolar (false) domain. Defaults to bipolar.
+		void SetBipolar(const bool bInBipolar);
+
+		// Set the waveform type of LFO generator
+		void SetType(const ELFO::Type InLFOType);
+
+		// Returns current waveform type of LFO generator
+		ELFO::Type GetType() const;
+
+		// Sets mode of LFO (Looping or one-shot.  Free not supported)
+		void SetMode(const ELFOMode::Type InLFOMode);
+
+		// Returns current mode of LFO
+		ELFOMode::Type GetMode() const;
+
+		// Sets a phase offset for the LFO.  Clamped to positive values only
+		void SetPhaseOffset(const float InOffset);
+
+		// Sets the exponential factor of LFO if type is "exponential". Clamped to positive, non-zero values.
+		void SetExponentialFactor(const float InExpFactor);
+
+		// Returns mod source's normal phase patch source
+		FPatchSource GetModSourceNormalPhase() const;
+
+		// Returns mod source's quad phase patch source
+		FPatchSource GetModSourceQuadPhase() const;
 
 	protected:
+		float ComputeLFO(float InputPhase, float* OutQuad = nullptr);
 
-		float ComputeLFO(const float InputPhase, float* OutQuad = nullptr);
+		// Returns initial phase, which differs between generators
+		// ensuring certain LFO shapes start on rising edge from zero-
+		// crossing by default (assuming user's provided phase offset is 0).
+		float GetInitPhase() const;
+
+		// Resets generator to initial phase
+		virtual void ResetPhase() override;
 
 		ELFO::Type LFOType;
 		ELFOMode::Type LFOMode;
@@ -76,10 +101,13 @@ namespace Audio
 		float ModScale;
 		float ModAdd;
 		float LastOutput;
+		float LoopCount;
 		float QuadLastOutput;
+		float PhaseOffset;
+
+		bool bBipolar;
 
 		FPatchSource ModNormalPhase;
 		FPatchSource ModQuadPhase;
 	};
-
-}
+} // namespace Audio
diff --git a/Engine/Source/Runtime/SignalProcessing/Public/DSP/Osc.h b/Engine/Source/Runtime/SignalProcessing/Public/DSP/Osc.h
--- a/Engine/Source/Runtime/SignalProcessing/Public/DSP/Osc.h
+++ b/Engine/Source/Runtime/SignalProcessing/Public/DSP/Osc.h
@@ -84,6 +84,9 @@ namespace Audio
 		// Resets the oscillator
 		virtual void Reset();
 
+		// Resets the phase of this oscillator to 0.0f
+		virtual void ResetPhase();
+
 		// Sets the gain of the oscillator
 		void SetGain(const float InGain) { Gain = InGain; }
 
@@ -116,9 +119,6 @@ namespace Audio
 		// Sets the LFO pulse width for square-wave type oscillators
 		void SetPulseWidth(const float InPulseWidth);
 
-		// Resets the phase of this oscillator to 0.0
-		void ResetPhase();
-
 		// Returns whether or not this oscillator is playing
 		bool IsPlaying() const { return bIsPlaying; }
 
