#ue5 - Type conversion fixes for MediaAssets, MovieScene, and SignalProcessing modules.

diff --git a/Engine/Source/Runtime/MediaAssets/Private/Assets/MediaTexture.cpp b/Engine/Source/Runtime/MediaAssets/Private/Assets/MediaTexture.cpp
--- a/Engine/Source/Runtime/MediaAssets/Private/Assets/MediaTexture.cpp
+++ b/Engine/Source/Runtime/MediaAssets/Private/Assets/MediaTexture.cpp
@@ -190,13 +190,13 @@ EMaterialValueType UMediaTexture::GetMaterialType() const
 
 float UMediaTexture::GetSurfaceWidth() const
 {
-	return Dimensions.X;
+	return (float)Dimensions.X;
 }
 
 
 float UMediaTexture::GetSurfaceHeight() const
 {
-	return Dimensions.Y;
+	return (float)Dimensions.Y;
 }
 
 
diff --git a/Engine/Source/Runtime/MediaAssets/Private/Misc/MediaTextureResource.cpp b/Engine/Source/Runtime/MediaAssets/Private/Misc/MediaTextureResource.cpp
--- a/Engine/Source/Runtime/MediaAssets/Private/Misc/MediaTextureResource.cpp
+++ b/Engine/Source/Runtime/MediaAssets/Private/Misc/MediaTextureResource.cpp
@@ -859,7 +859,7 @@ void FMediaTextureResource::ConvertSample(const TSharedPtr<IMediaTextureSample,
 		CommandList.BeginRenderPass(RPInfo, TEXT("ConvertMedia"));
 		{
 			CommandList.ApplyCachedRenderTargets(GraphicsPSOInit);
-			CommandList.SetViewport(0, 0, 0.0f, OutputDim.X, OutputDim.Y, 1.0f);
+			CommandList.SetViewport(0, 0, 0.0f, (float)OutputDim.X, (float)OutputDim.Y, 1.0f);
 
 			GraphicsPSOInit.DepthStencilState = TStaticDepthStencilState<false, CF_Always>::GetRHI();
 			GraphicsPSOInit.RasterizerState = TStaticRasterizerState<>::GetRHI();
@@ -1066,7 +1066,7 @@ void FMediaTextureResource::ConvertSample(const TSharedPtr<IMediaTextureSample,
 			FBufferRHIRef VertexBuffer = CreateTempMediaVertexBuffer();
 			CommandList.SetStreamSource(0, VertexBuffer, 0);
 			// set viewport to RT size
-			CommandList.SetViewport(0, 0, 0.0f, OutputDim.X, OutputDim.Y, 1.0f);
+			CommandList.SetViewport(0, 0, 0.0f, (float)OutputDim.X, (float)OutputDim.Y, 1.0f);
 
 			CommandList.DrawPrimitive(0, 2, 1);
 		}
@@ -1195,7 +1195,7 @@ void FMediaTextureResource::CopyFromExternalTexture(const TSharedPtr <IMediaText
 			const FIntPoint OutputDim = Sample->GetOutputDim();
 
 			CommandList.ApplyCachedRenderTargets(GraphicsPSOInit);
-			CommandList.SetViewport(0, 0, 0.0f, OutputDim.X, OutputDim.Y, 1.0f);
+			CommandList.SetViewport(0, 0, 0.0f, (float)OutputDim.X, (float)OutputDim.Y, 1.0f);
 			
 			GraphicsPSOInit.DepthStencilState = TStaticDepthStencilState<false, CF_Always>::GetRHI();
 			GraphicsPSOInit.RasterizerState = TStaticRasterizerState<>::GetRHI();
@@ -1218,7 +1218,7 @@ void FMediaTextureResource::CopyFromExternalTexture(const TSharedPtr <IMediaText
 			FBufferRHIRef VertexBuffer = CreateTempMediaVertexBuffer();
 			CommandList.SetStreamSource(0, VertexBuffer, 0);
 			// set viewport to RT size
-			CommandList.SetViewport(0, 0, 0.0f, OutputDim.X, OutputDim.Y, 1.0f);
+			CommandList.SetViewport(0, 0, 0.0f, (float)OutputDim.X, (float)OutputDim.Y, 1.0f);
 
 			CommandList.DrawPrimitive(0, 2, 1);
 		}
@@ -1280,7 +1280,7 @@ void FMediaTextureResource::CreateOutputRenderTarget(const FIntPoint & InDim, EP
 		OutputCreateFlags |= TexCreate_GenerateMipCapable;
 
 		// Make sure we only set a number of mips that actually makes sense, given the sample size
-		uint8 MaxMips = FGenericPlatformMath::FloorToInt(FGenericPlatformMath::Log2(static_cast<float>(FGenericPlatformMath::Min(InDim.X, InDim.Y))));
+		uint8 MaxMips = (uint8)(FMath::Min(255, FGenericPlatformMath::FloorToInt(FGenericPlatformMath::Log2(static_cast<float>(FGenericPlatformMath::Min(InDim.X, InDim.Y))))));
 		InNumMips = FMath::Min(InNumMips, MaxMips);
 	}
 
diff --git a/Engine/Source/Runtime/MovieScene/Public/Evaluation/MovieSceneTimeWarping.h b/Engine/Source/Runtime/MovieScene/Public/Evaluation/MovieSceneTimeWarping.h
--- a/Engine/Source/Runtime/MovieScene/Public/Evaluation/MovieSceneTimeWarping.h
+++ b/Engine/Source/Runtime/MovieScene/Public/Evaluation/MovieSceneTimeWarping.h
@@ -72,7 +72,7 @@ struct FMovieSceneTimeWarping
 		else
 		{
 			const FFrameNumber WarpLength = Length();
-			const FFrameNumber WarpsOffset = WarpLength * WarpCount;
+			const FFrameNumber WarpsOffset = WarpLength * (float)WarpCount;
 			return FMovieSceneTimeTransform(WarpsOffset);
 		}
 	}
diff --git a/Engine/Source/Runtime/SignalProcessing/Public/SampleBuffer.h b/Engine/Source/Runtime/SignalProcessing/Public/SampleBuffer.h
--- a/Engine/Source/Runtime/SignalProcessing/Public/SampleBuffer.h
+++ b/Engine/Source/Runtime/SignalProcessing/Public/SampleBuffer.h
@@ -76,11 +76,11 @@ namespace Audio
 			RawPCMData.Reset(NumSamples);
 			RawPCMData.AddUninitialized(NumSamples);
 
-			if (TIsSame<SampleType, float>::Value)
+			if constexpr(TIsSame<SampleType, float>::Value)
 			{
 				FMemory::Memcpy(RawPCMData.GetData(), InBufferPtr, NumSamples * sizeof(float));
 			}
-			else if (TIsSame<SampleType, int16>::Value)
+			else if constexpr(TIsSame<SampleType, int16>::Value)
 			{
 				// Convert from float to int:
 				for (int32 SampleIndex = 0; SampleIndex < NumSamples; SampleIndex++)
@@ -109,11 +109,11 @@ namespace Audio
 			RawPCMData.Reset(NumSamples);
 			RawPCMData.AddUninitialized(NumSamples);
 
-			if (TIsSame<SampleType, int16>::Value)
+			if constexpr(TIsSame<SampleType, int16>::Value)
 			{
 				FMemory::Memcpy(RawPCMData.GetData(), InBufferPtr, NumSamples * sizeof(int16));
 			}
-			else if (TIsSame<SampleType, float>::Value)
+			else if constexpr(TIsSame<SampleType, float>::Value)
 			{
 				// Convert from int to float:
 				for (int32 SampleIndex = 0; SampleIndex < NumSamples; SampleIndex++)
@@ -161,12 +161,12 @@ namespace Audio
 			RawPCMData.Reset(NumSamples);
 			RawPCMData.AddUninitialized(NumSamples);
 
-			if (TIsSame<SampleType, OtherSampleType>::Value)
+			if constexpr(TIsSame<SampleType, OtherSampleType>::Value)
 			{
 				// If buffers are of the same type, copy over:
 				FMemory::Memcpy(RawPCMData.GetData(), Other.RawPCMData.GetData(), NumSamples * sizeof(SampleType));
 			}
-			else if (TIsSame<SampleType, int16>::Value && TIsSame<OtherSampleType, float>::Value)
+			else if constexpr(TIsSame<SampleType, int16>::Value && TIsSame<OtherSampleType, float>::Value)
 			{
 				// Convert from float to int:
 				for (int32 SampleIndex = 0; SampleIndex < NumSamples; SampleIndex++)
@@ -174,7 +174,7 @@ namespace Audio
 					RawPCMData[SampleIndex] = (int16)(Other.RawPCMData[SampleIndex] * 32767.0f);
 				}
 			}
-			else if (TIsSame<SampleType, float>::Value && TIsSame<OtherSampleType, int16>::Value)
+			else if constexpr(TIsSame<SampleType, float>::Value && TIsSame<OtherSampleType, int16>::Value)
 			{
 				// Convert from int to float:
 				for (int32 SampleIndex = 0; SampleIndex < NumSamples; SampleIndex++)
@@ -187,7 +187,7 @@ namespace Audio
 				// for any other types, we don't know how to explicitly convert, so we fall back to casts:
 				for (int32 SampleIndex = 0; SampleIndex < NumSamples; SampleIndex++)
 				{
-					RawPCMData[SampleIndex] = Other.RawPCMData[SampleIndex];
+					RawPCMData[SampleIndex] = (SampleType)Other.RawPCMData[SampleIndex];
 				}
 			}
 
@@ -215,13 +215,13 @@ namespace Audio
 		{
 			int32 StartIndex = RawPCMData.AddUninitialized(InNumSamples);
 
-			if (TIsSame<SampleType, OtherSampleType>::Value)
+			if constexpr(TIsSame<SampleType, OtherSampleType>::Value)
 			{
 				FMemory::Memcpy(&RawPCMData[StartIndex], InputBuffer, InNumSamples * sizeof(SampleType));
 			}
 			else
 			{
-				if (TIsSame<SampleType, int16>::Value && TIsSame<OtherSampleType, float>::Value)
+				if constexpr(TIsSame<SampleType, int16>::Value && TIsSame<OtherSampleType, float>::Value)
 				{
 					// Convert from float to int:
 					for (int32 SampleIndex = 0; SampleIndex < InNumSamples; SampleIndex++)
@@ -229,7 +229,7 @@ namespace Audio
 						RawPCMData[StartIndex + SampleIndex] = (int16)(InputBuffer[SampleIndex] * 32767.0f);
 					}
 				}
-				else if (TIsSame<SampleType, float>::Value && TIsSame<OtherSampleType, int16>::Value)
+				else if constexpr(TIsSame<SampleType, float>::Value && TIsSame<OtherSampleType, int16>::Value)
 				{
 					// Convert from int to float:
 					for (int32 SampleIndex = 0; SampleIndex < InNumSamples; SampleIndex++)
@@ -358,7 +358,7 @@ namespace Audio
 
 		void Clamp(float Ceiling = 1.0f)
 		{
-			if (TIsSame<SampleType, float>::Value)
+			if constexpr(TIsSame<SampleType, float>::Value)
 			{
 				// Float case:
 				float ClampMin = Ceiling * -1.0f;
@@ -368,7 +368,7 @@ namespace Audio
 					RawPCMData[SampleIndex] = static_cast<SampleType>(FMath::Clamp<float>(RawPCMData[SampleIndex], ClampMin, Ceiling));
 				}
 			}
-			else if (TIsSame<SampleType, int16>::Value)
+			else if constexpr(TIsSame<SampleType, int16>::Value)
 			{
 				// int16 case:
 				Ceiling = FMath::Clamp(Ceiling, 0.0f, 1.0f);
@@ -480,7 +480,7 @@ namespace Audio
 			{
 				float SampleA, SampleB;
 
-				if (TIsSame<SampleType, float>::Value)
+				if constexpr(TIsSame<SampleType, float>::Value)
 				{
 					SampleA = RawPCMData[(WholeThisIndex * NumChannels) + i];
 					SampleB = RawPCMData[(WholeNextIndex * NumChannels) + i];
