Added missing scale degree to dorian mode (scale to note array node).

diff --git a/Engine/Source/Runtime/SignalProcessing/Private/MidiNoteQuantizer.cpp b/Engine/Source/Runtime/SignalProcessing/Private/MidiNoteQuantizer.cpp
--- a/Engine/Source/Runtime/SignalProcessing/Private/MidiNoteQuantizer.cpp
+++ b/Engine/Source/Runtime/SignalProcessing/Private/MidiNoteQuantizer.cpp
@@ -150,7 +150,7 @@ namespace Audio
 		, { EMusicalScale::Scale::Dominant7th_Mixolydian,
 			{{ n1, n2, n3, n4, n5, n6, b7, n8 }, { n1, n2, n3, n5, b7, n8 }}}
 		, { EMusicalScale::Scale::Minor_Dorian,
-			{{ n1, n2, b3, n4, n5, b7, n8 }, { n1, n2, b3, n5, b7, n8 }}}
+			{{ n1, n2, b3, n4, n5, n6, b7, n8 }, { n1, n2, b3, n5, b7, n8 }}}
 		, { EMusicalScale::Scale::HalfDiminished_Locrian,
 			{{ n1, b2, b3, n4, b5, b6, b7, n8 }, { n1, b3, b5, b7, n8 }}}
 		, { EMusicalScale::Scale::Diminished,
