Fix for unit test failure in sine wavetable oscillator.
#rnx

diff --git a/Engine/Source/Runtime/SignalProcessing/Private/SineWaveTableOsc.cpp b/Engine/Source/Runtime/SignalProcessing/Private/SineWaveTableOsc.cpp
--- a/Engine/Source/Runtime/SignalProcessing/Private/SineWaveTableOsc.cpp
+++ b/Engine/Source/Runtime/SignalProcessing/Private/SineWaveTableOsc.cpp
@@ -89,7 +89,7 @@ namespace Audio
 		for (int32 SampleIndex = 0; SampleIndex < NumSamples; ++SampleIndex)
 		{
 			// Interpolate between two samples
-			const int32 ReadIndexPrev = FMath::Clamp((int32)ReadIndex, 0, WaveTableBuffer.Num() - 2);
+			const int32 ReadIndexPrev = FMath::Clamp((int32)ReadIndex, 0, WaveTableBuffer.Num() - 1);
 			const float Alpha = FMath::Clamp(ReadIndex - (float)ReadIndexPrev, 0.f, 1.f);
 
 			const int32 ReadIndexNext = (ReadIndexPrev + 1) % WaveTableBuffer.Num();
