Audio Bus Reader/Writer Nodes Crackle Fix

Audio Bus Reader Node more accurately calculates how many samples it needs before it can start popping.
Reader node has a CVar timeout on how long it'll wait for samples (whereas it used to calculate a huge timeout value that I suspect was erroneous).  In theory the samples should always become available even with an infinite timeout.
The Metasound generator's BuilderTask's completion is ensured at the beginning of FMetasoundGenerator::OnGenerateAudio.  This prevents the mixer, writer and reader nodes from getting out of sync in terms of samples output and input.
FMixerAudioBus::MixBuffer doesn't try to cache any samples and just uses PopAudio and PushAudio directly on its mixer and splitter.
FPatchOutput::WaitUntilNumSamplesAvailable no longer has a race condition and an attempt has been made to improve thread-safety in FPatchOutput.

diff --git a/Engine/Plugins/Runtime/Metasound/Source/MetasoundEngine/Private/MetasoundAudioBusReaderNode.cpp b/Engine/Plugins/Runtime/Metasound/Source/MetasoundEngine/Private/MetasoundAudioBusReaderNode.cpp
--- a/Engine/Plugins/Runtime/Metasound/Source/MetasoundEngine/Private/MetasoundAudioBusReaderNode.cpp
+++ b/Engine/Plugins/Runtime/Metasound/Source/MetasoundEngine/Private/MetasoundAudioBusReaderNode.cpp
@@ -16,6 +16,13 @@
 
 #define LOCTEXT_NAMESPACE "MetasoundAudioBusNode"
 
+static int32 AudioBusReaderNodePatchWaitTimeout = 3;
+FAutoConsoleVariableRef CVarAudioBusReaderNodePatchWaitTimeout(
+	TEXT("au.BusReaderPatchWaitTimeout"),
+	AudioBusReaderNodePatchWaitTimeout,
+	TEXT("The maximum amount of time the audio bus reader node will wait for its patch output to receive samples."),
+	ECVF_Default);
+
 namespace Metasound
 {
 	namespace AudioBusReaderNode
@@ -135,10 +142,6 @@ namespace Metasound
 			}
 		}
 
-		virtual ~TAudioBusReaderOperator()
-		{
-		}
-
 		virtual FDataReferenceCollection GetInputs() const override
 		{
 			using namespace AudioBusReaderNode;
@@ -172,33 +175,32 @@ namespace Metasound
 				return;
 			}
 
-			// We want to write out zeroed audio until there is enough audio queued up in the patch output to avoid any underruns due to thread timing.
-			// From testing, waiting for the first buffer for 3 times the size of a single metasound block was sufficient to avoid pops. 
-			// Once we have started popping audio off the patch output, we want to keep it popping. Unless there are CPU underruns, there should be enough
-			// runway to keep it from underruning. We employ a little bit of a wait with a time out just to decrease the chance that there will be any missed buffers.
-			// In practice, this wait should rarely, if ever, actually cause a wait in this metasound execute.
-
 			bool bPerformPop = false;
+
 			if (bFirstBlock)
 			{
-				// Tuned amount to wait for first rendered block to account for any timing issues between the audio render thread and the metasound task
-				constexpr int32 FirstBlockBufferCount = 3;
-				if (AudioBusPatchOutput->GetNumSamplesAvailable() > FirstBlockBufferCount * NumSamplesToPop)
+				// Ensure there are enough samples in the patch output to support the maximum metasound executions the mixer requires to fill its output frames before popping begins.
+				if (uint32(AudioBusPatchOutput->GetNumSamplesAvailable()) >= FMath::DivideAndRoundUp(AudioMixerOutputFrames, BlockSizeFrames) * BlockSizeFrames * AudioBusChannels)
 				{
-					bFirstBlock = false;
 					bPerformPop = true;
 
 					InterleavedBuffer.Reset();
 					InterleavedBuffer.AddUninitialized(NumSamplesToPop);
 				}
 			}
-			else if (AudioBusPatchOutput->WaitUntilNumSamplesAvailable(NumSamplesToPop, static_cast<uint32>(0.5f * BlockSizeFrames * SampleRate)))
+			else
 			{
-				bPerformPop = true;
+				// Give input patch samples some time to arrive.
+				if (AudioBusPatchOutput->WaitUntilNumSamplesAvailable(NumSamplesToPop, uint32(AudioBusReaderNodePatchWaitTimeout)))
+				{
+					bPerformPop = true;
+				}
 			}
 
 			if (bPerformPop)
 			{
+				bFirstBlock = false;
+
 				// Pop off the interleaved data from the audio bus
 				int32 SamplesPopped = AudioBusPatchOutput->PopAudio(InterleavedBuffer.GetData(), NumSamplesToPop, false);
 				if (SamplesPopped < NumSamplesToPop)
diff --git a/Engine/Plugins/Runtime/Metasound/Source/MetasoundEngine/Private/MetasoundAudioBusWriterNode.cpp b/Engine/Plugins/Runtime/Metasound/Source/MetasoundEngine/Private/MetasoundAudioBusWriterNode.cpp
--- a/Engine/Plugins/Runtime/Metasound/Source/MetasoundEngine/Private/MetasoundAudioBusWriterNode.cpp
+++ b/Engine/Plugins/Runtime/Metasound/Source/MetasoundEngine/Private/MetasoundAudioBusWriterNode.cpp
@@ -136,10 +136,6 @@ namespace Metasound
 			}
 		}
 
-		virtual ~TAudioBusWriterOperator()
-		{
-		}
-
 		virtual FDataReferenceCollection GetInputs() const override
 		{
 			using namespace AudioBusWriterNode;
@@ -201,6 +197,16 @@ namespace Metasound
 				}
 			}
 
+			if (bFirstBlock)
+			{
+				bFirstBlock = false;
+				if (AudioMixerOutputFrames != BlockSizeFrames)
+				{
+					// Ensure there will be enough samples in the patch input to support the maximum metasound executions the mixer requires to fill its output frames after the next push.
+					AudioBusPatchInput.PushAudio(nullptr, (FMath::DivideAndRoundUp(FMath::Max(AudioMixerOutputFrames, BlockSizeFrames), FMath::Min(AudioMixerOutputFrames, BlockSizeFrames)) - 1) * BlockSizeFrames * AudioBusChannels);
+				}
+			}
+
 			// Pushes the interleaved data to the audio bus
 			int32 SamplesPushed = AudioBusPatchInput.PushAudio(InterleavedBuffer.GetData(), InterleavedBuffer.Num());
 			if (SamplesPushed < InterleavedBuffer.Num())
@@ -220,6 +226,7 @@ namespace Metasound
 		Audio::FPatchInput AudioBusPatchInput;
 		uint32 AudioBusChannels = INDEX_NONE;
 		int32 BlockSizeFrames = 0;
+		bool bFirstBlock = true;
 	};
 
 	template<uint32 NumChannels>
diff --git a/Engine/Source/Runtime/AudioMixer/Private/AudioMixerBus.cpp b/Engine/Source/Runtime/AudioMixer/Private/AudioMixerBus.cpp
--- a/Engine/Source/Runtime/AudioMixer/Private/AudioMixerBus.cpp
+++ b/Engine/Source/Runtime/AudioMixer/Private/AudioMixerBus.cpp
@@ -88,25 +88,13 @@ namespace Audio
 	void FMixerAudioBus::MixBuffer()
 	{
 		// Mix the patch mixer's inputs into the source data
-		const int32 MaxNumSamplesToPop = PatchMixer.MaxNumberOfSamplesThatCanBePopped();
 		const int32 NumSamples = NumFrames * NumChannels;
 		const int32 NumOutputFrames = SourceManager->GetNumOutputFrames();
 
 		FAlignedFloatBuffer& MixBuffer = MixedSourceData[CurrentBufferIndex];
 		float* BusDataBufferPtr = MixBuffer.GetData();
 
-		if (MaxNumSamplesToPop < NumSamples)
-		{
-			// Zero buffer when there are no inputs connected,
-			// or one of the inputs has not pushed any audio yet.
-			FMemory::Memzero(BusDataBufferPtr, NumSamples * sizeof(float));
-		}
-		else
-		{
-			// Mix down inputs
-			int32 PopResult = PatchMixer.PopAudio(BusDataBufferPtr, NumSamples, false);
-			check(PopResult == NumSamples);
-		}
+		PatchMixer.PopAudio(BusDataBufferPtr, NumSamples, false);
 
 		for (int32 BusSendType = 0; BusSendType < (int32)EBusSendType::Count; ++BusSendType)
 		{
@@ -156,38 +144,7 @@ namespace Audio
 		}
 
 		// Send the mix to the patch splitter's outputs
-		int32 NumSamplesToPush = PatchSplitter.MaxNumberOfSamplesThatCanBePushed();
-		if (NumSamplesToPush <= 0)
-		{
-			// Early exit when there are no outputs.
-			return;
-		}
-
-		int32 SampleCacheSize = SampleCache.Num();
-		if (0 <= SampleCacheSize)
-		{
-			float* SampleCacheBufferPtr = SampleCache.GetData();
-			PatchSplitter.PushAudio(SampleCacheBufferPtr, FMath::Min(SampleCacheSize, NumSamplesToPush));
-
-			if (NumSamplesToPush < SampleCacheSize)
-			{
-				const int32 NumSamplesToCache = SampleCacheSize - NumSamplesToPush;
-				FMemory::Memmove(SampleCacheBufferPtr, &SampleCache[NumSamplesToPush], NumSamplesToCache);
-				SampleCache.SetNum(NumSamplesToCache, false);
-				SampleCache.Append(MixBuffer);
-				return;
-			}
-
-			NumSamplesToPush -= SampleCacheSize;
-			SampleCache.Reset();
-		}
-
-		PatchSplitter.PushAudio(BusDataBufferPtr, FMath::Min(NumSamples, NumSamplesToPush));
-
-		if (NumSamplesToPush < NumSamples)
-		{
-			SampleCache.Append(&MixBuffer[NumSamplesToPush], NumSamples - NumSamplesToPush);
-		}
+		PatchSplitter.PushAudio(BusDataBufferPtr, NumSamples);
 	}
 
 	void FMixerAudioBus::CopyCurrentBuffer(Audio::FAlignedFloatBuffer& InChannelMap, int32 InNumOutputChannels, FAlignedFloatBuffer& OutBuffer, int32 NumOutputFrames) const
diff --git a/Engine/Source/Runtime/AudioMixer/Public/AudioMixerBus.h b/Engine/Source/Runtime/AudioMixer/Public/AudioMixerBus.h
--- a/Engine/Source/Runtime/AudioMixer/Public/AudioMixerBus.h
+++ b/Engine/Source/Runtime/AudioMixer/Public/AudioMixerBus.h
@@ -108,7 +108,6 @@ namespace Audio
 		// Multiple places can produce and consume from audio buses
 		Audio::FPatchMixer PatchMixer;
 		Audio::FPatchSplitter PatchSplitter;
-		FAlignedFloatBuffer SampleCache;
 
 		// Was created manually, not via source buses.
 		bool bIsAutomatic;
diff --git a/Engine/Source/Runtime/SignalProcessing/Private/MultithreadedPatching.cpp b/Engine/Source/Runtime/SignalProcessing/Private/MultithreadedPatching.cpp
--- a/Engine/Source/Runtime/SignalProcessing/Private/MultithreadedPatching.cpp
+++ b/Engine/Source/Runtime/SignalProcessing/Private/MultithreadedPatching.cpp
@@ -25,8 +25,7 @@ namespace Audio
 		, PreviousGain(InGain)
 		, PatchID(++PatchIDCounter)
 		, NumAliveInputs(0)
-		, SamplesFilledEvent(nullptr)
-		, NumSamplesToWaitFor(0)
+		, SamplesPushedEvent(nullptr)
 	{
 
 	}
@@ -38,17 +37,15 @@ namespace Audio
 		, PreviousGain(0.0f)
 		, PatchID(INDEX_NONE)
 		, NumAliveInputs(0)
-		, SamplesFilledEvent(nullptr)
-		, NumSamplesToWaitFor(0)
+		, SamplesPushedEvent(nullptr)
 	{
 	}
 
 	FPatchOutput::~FPatchOutput()
 	{
-		if (SamplesFilledEvent)
+		if (FEvent* Event = SamplesPushedEvent.exchange(nullptr, std::memory_order_acq_rel))
 		{
-			FPlatformProcess::ReturnSynchEventToPool(SamplesFilledEvent);
-			SamplesFilledEvent = nullptr;
+			FPlatformProcess::ReturnSynchEventToPool(Event);
 		}
 	}
 
@@ -88,18 +85,17 @@ namespace Audio
 
 	bool FPatchOutput::IsInputStale() const
 	{
-		return NumAliveInputs == 0;
+		return NumAliveInputs.load(std::memory_order_relaxed) == 0;
 	}
 
 	int32 FPatchOutput::PushAudioToInternalBuffer(const float* InBuffer, int32 NumSamples)
 	{
-		const int32 NumSamplesPushed = InternalBuffer.Push(InBuffer, NumSamples);
+		const int32 NumSamplesPushed = InBuffer ? InternalBuffer.Push(InBuffer, NumSamples) : InternalBuffer.PushZeros(NumSamples);
 
 		// Check to see if we need to notify anybody waiting on this patch output getting filled
-		if (SamplesFilledEvent && NumSamplesToWaitFor > 0 && GetNumSamplesAvailable() >= NumSamplesToWaitFor)
+		if (FEvent* Event = SamplesPushedEvent.load(std::memory_order_acquire))
 		{
-			SamplesFilledEvent->Trigger();
-			NumSamplesToWaitFor = 0;
+			Event->Trigger();
 		}
 
 		const int32 CurrSamplesAvailable = GetNumSamplesAvailable();
@@ -155,19 +151,60 @@ namespace Audio
 
 	bool FPatchOutput::WaitUntilNumSamplesAvailable(int32 InNumSamplesToWaitFor, uint32 TimeOutMilliseconds)
 	{
-		// No need to wait if we have that many available already
-		if (GetNumSamplesAvailable() >= InNumSamplesToWaitFor)
+		// Samples are ready if there are enough of them available, or the input is stale.
+		// Conceptually, a stale input contributes silence, so its samples are known.
+		if (IsInputStale() || GetNumSamplesAvailable() >= InNumSamplesToWaitFor)
 		{
 			return true;
 		}
-		NumSamplesToWaitFor = InNumSamplesToWaitFor;
-		
-		if (!SamplesFilledEvent)
+
+		// Samples are not ready if the timeout is zero or the
+		// internal buffer can't hold the number of requested samples.
+		if (!TimeOutMilliseconds || InternalBuffer.GetCapacity() < uint32(InNumSamplesToWaitFor))
+		{
+			return false;
+		}
+
+		// Determine if an event is ready to signal when
+		// samples are submitted.  Provide one if not.
+		FEvent* Event = SamplesPushedEvent.load(std::memory_order_acquire);
+		if (!Event)
 		{
-			SamplesFilledEvent = FPlatformProcess::GetSynchEventFromPool(false);
+			Event = FPlatformProcess::GetSynchEventFromPool(false);
+			SamplesPushedEvent.store(Event, std::memory_order_release);
+		}
+
+		// Calculate when the timeout period will end.
+		double CurrentTime = FPlatformTime::Seconds();
+		double WaitStopTime = CurrentTime + 0.001 * TimeOutMilliseconds;
+
+		// First determine if samples are available.
+		// They might have become available before the event was ready to signal.
+		while (!IsInputStale() && GetNumSamplesAvailable() < InNumSamplesToWaitFor)
+		{
+			// Wait for the next sample submission.
+			Event->Wait(TimeOutMilliseconds);
+
+			// If the timeout is not infinite, update it.
+			if (TimeOutMilliseconds != MAX_uint32)
+			{
+				CurrentTime = FPlatformTime::Seconds();
+
+				// Return if no time remains, indicating whether samples are ready.
+				double TimeOutRemainder = (WaitStopTime - CurrentTime) * 1000.;
+				if (!(TimeOutRemainder > 0.))
+				{
+					return IsInputStale() || GetNumSamplesAvailable() >= InNumSamplesToWaitFor;
+				}
+
+				// Truncating to an integer will cause a spin-wait
+				// when there's less than a millisecond remaining.
+				TimeOutMilliseconds = uint32(TimeOutRemainder);
+			}
 		}
-	
-		return SamplesFilledEvent->Wait(TimeOutMilliseconds);
+
+		// Samples are ready.
+		return true;
 	}
 
 	FPatchInput::FPatchInput(const FPatchOutputStrongPtr& InOutput)
@@ -175,7 +212,7 @@ namespace Audio
 	{
 		if (OutputHandle.IsValid())
 		{
-			OutputHandle->NumAliveInputs++;
+			OutputHandle->NumAliveInputs.fetch_add(1, std::memory_order_relaxed);
 		}
 	}
 
@@ -184,7 +221,7 @@ namespace Audio
 	{
 		if (OutputHandle.IsValid())
 		{
-			OutputHandle->NumAliveInputs++;
+			OutputHandle->NumAliveInputs.fetch_add(1, std::memory_order_relaxed);
 		}
 	}
 
@@ -197,24 +234,15 @@ namespace Audio
 
 	FPatchInput& FPatchInput::operator=(const FPatchInput& Other)
 	{
-		OutputHandle = Other.OutputHandle;
-		PushCallsCounter = 0;
-		
-		if (OutputHandle.IsValid())
-		{
-			OutputHandle->NumAliveInputs++;
-		}
-
+		FPatchInput NewPatchInput(Other);
+		Swap(NewPatchInput, *this);
 		return *this;
 	}
 
 	FPatchInput& FPatchInput::operator=(FPatchInput&& Other)
 	{
-		OutputHandle = MoveTemp(Other.OutputHandle);
-
-		PushCallsCounter = Other.PushCallsCounter;
-		Other.PushCallsCounter = 0;
-
+		FPatchInput NewPatchInput(MoveTemp(Other));
+		Swap(NewPatchInput, *this);
 		return *this;
 	}
 
@@ -222,7 +250,14 @@ namespace Audio
 	{
 		if (OutputHandle.IsValid())
 		{
-			OutputHandle->NumAliveInputs--;
+			if (OutputHandle->NumAliveInputs.fetch_sub(1, std::memory_order_relaxed) == 1)
+			{
+				if (FEvent* Event = OutputHandle->SamplesPushedEvent.load(std::memory_order_acquire))
+				{
+					// Signal that the sample status should be checked.
+					Event->Trigger();
+				}
+			}
 		}
 	}
 
@@ -322,6 +357,10 @@ namespace Audio
 	int32 FPatchMixer::Num()
 	{
 		FScopeLock ScopeLock(&CurrentPatchesCriticalSection);
+
+		CleanUpDisconnectedPatches();
+		ConnectNewPatches();
+
 		return CurrentInputs.Num();
 	}
 
@@ -351,6 +390,53 @@ namespace Audio
 		return SmallestNumSamplesBuffered;
 	}
 
+	bool FPatchMixer::WaitUntilNumSamplesAvailable(int32 NumSamples, uint32 TimeOutMilliseconds)
+	{
+		FScopeLock ScopeLock(&CurrentPatchesCriticalSection);
+
+		CleanUpDisconnectedPatches();
+		ConnectNewPatches();
+
+		if (CurrentInputs.IsEmpty())
+		{
+			// Samples are ready if there are enough of them available, or the input is stale.
+			// Conceptually, a stale input contributes silence, so its samples are known.
+			return true;
+		}
+
+		// Calculate when the timeout period will end.
+		double CurrentTime = FPlatformTime::Seconds();
+		double WaitStopTime = CurrentTime + 0.001 * TimeOutMilliseconds;
+		for (FPatchOutputStrongPtr& Output : CurrentInputs)
+		{
+			bool bSamplesAvailable;
+			do
+			{
+				bSamplesAvailable = Output->WaitUntilNumSamplesAvailable(NumSamples, TimeOutMilliseconds);
+
+				// If the timeout is not infinite, update it.
+				if (TimeOutMilliseconds != MAX_uint32)
+				{
+					CurrentTime = FPlatformTime::Seconds();
+
+					// Return if no time remains and samples are not ready.
+					double TimeOutRemainder = (WaitStopTime - CurrentTime) * 1000.;
+					if (!bSamplesAvailable && !(TimeOutRemainder > 0.))
+					{
+						return false;
+					}
+
+					// If samples were ready, but the timeout has expired,
+					// the wait could still succeed if all the subsequent inputs are ready.
+					TimeOutMilliseconds = TimeOutRemainder > 0. ? uint32(TimeOutRemainder) : 0;
+				}
+			}
+			while (!bSamplesAvailable);
+		}
+
+		return true;
+	}
+
 	void FPatchMixer::DisconnectAllInputs()
 	{
 		FScopeLock ScopeLock(&CurrentPatchesCriticalSection);
diff --git a/Engine/Source/Runtime/SignalProcessing/Public/DSP/MultithreadedPatching.h b/Engine/Source/Runtime/SignalProcessing/Public/DSP/MultithreadedPatching.h
--- a/Engine/Source/Runtime/SignalProcessing/Public/DSP/MultithreadedPatching.h
+++ b/Engine/Source/Runtime/SignalProcessing/Public/DSP/MultithreadedPatching.h
@@ -50,7 +50,7 @@ namespace Audio
 
 		/** Pauses the current thread until there are the given number of samples available to pop. Will return true if it succeeded, false if it timed out. */
 		bool WaitUntilNumSamplesAvailable(int32 NumSamples, uint32 TimeOutMilliseconds = MAX_uint32);
-		
+
 		/** Returns true if the input for this patch has been destroyed. */
 		bool IsInputStale() const;
 
@@ -75,12 +75,11 @@ namespace Audio
 		int32 PatchID;
 
 		// Counter that is incremented/decremented to allow FPatchInput to be copied around safely.
-		int32 NumAliveInputs;
+		std::atomic<int32> NumAliveInputs;
 
 		// Event to pause the current thread until a given number of samples has been filled
-		FEvent* SamplesFilledEvent;
-		int32 NumSamplesToWaitFor;
-		
+		std::atomic<FEvent*> SamplesPushedEvent;
+
 		static TAtomic<int32> PatchIDCounter;
 	};
 
@@ -106,7 +105,8 @@ namespace Audio
 
 		~FPatchInput();
 
-		/** pushes audio from InBuffer to the corresponding FPatchOutput.
+		/** Pushes audio from InBuffer to the corresponding FPatchOutput.
+		 *  Pushes zeros if InBuffer is nullptr.
 		 *  Returns how many samples were able to be pushed, or -1 if the output was disconnected.
 		 */
 		int32 PushAudio(const float* InBuffer, int32 NumSamples);
@@ -153,6 +153,9 @@ namespace Audio
 		/** This function call gets the maximum number of samples that's safe to pop, based on the thread with the least amount of samples buffered. Thread safe, but blocks for PopAudio. */
 		int32 MaxNumberOfSamplesThatCanBePopped();
 
+		/** Pauses the current thread until there are the given number of samples available to pop. Will return true if it succeeded, false if it timed out. */
+		bool WaitUntilNumSamplesAvailable(int32 NumSamples, uint32 TimeOutMilliseconds = MAX_uint32);
+
 		/** Disconnect everything currently connected to this mixer. */
 		void DisconnectAllInputs();
 
