Nonunity build fix

include HAL/PlatformTime.h

diff --git a/Engine/Source/Runtime/SignalProcessing/Private/MultithreadedPatching.cpp b/Engine/Source/Runtime/SignalProcessing/Private/MultithreadedPatching.cpp
--- a/Engine/Source/Runtime/SignalProcessing/Private/MultithreadedPatching.cpp
+++ b/Engine/Source/Runtime/SignalProcessing/Private/MultithreadedPatching.cpp
@@ -5,6 +5,7 @@
 #include "DSP/FloatArrayMath.h"
 #include "HAL/IConsoleManager.h"
 #include "HAL/PlatformProcess.h"
+#include "HAL/PlatformTime.h"
 #include "HAL/Event.h"
 
 static int32 MultithreadedPatchingPushCallsPerOutputCleanupCheckCVar = 256;
