Fix for unit test crashing on ARM machines.

diff --git a/Engine/Source/Runtime/SignalProcessing/Private/MelScale.cpp b/Engine/Source/Runtime/SignalProcessing/Private/MelScale.cpp
--- a/Engine/Source/Runtime/SignalProcessing/Private/MelScale.cpp
+++ b/Engine/Source/Runtime/SignalProcessing/Private/MelScale.cpp
@@ -63,7 +63,7 @@ namespace Audio
 		// Convert frequency bounds to mel
 		float MinMel = HzToMel(InSettings.KernelMinCenterFreq);
 		float MaxMel = HzToMel(InSettings.KernelMaxCenterFreq);
-		float DeltaMel = (MaxMel - MinMel) / (InSettings.NumBands - 1);
+		float DeltaMel = (MaxMel - MinMel) / FMath::Max(1, (InSettings.NumBands - 1));
 
 		// Only need to look at half of spectrum because magnitudes are mirrored.
 		const int32 NumUsefulFrequencyBins = InFFTSize / 2 + 1;
