Fix race condition in FPatchOutput::WaitUntilNumSamplesAvailable when setting up the sample pushed event.

diff --git a/Engine/Source/Runtime/SignalProcessing/Private/MultithreadedPatching.cpp b/Engine/Source/Runtime/SignalProcessing/Private/MultithreadedPatching.cpp
--- a/Engine/Source/Runtime/SignalProcessing/Private/MultithreadedPatching.cpp
+++ b/Engine/Source/Runtime/SignalProcessing/Private/MultithreadedPatching.cpp
@@ -171,8 +171,17 @@ namespace Audio
 		FEvent* Event = SamplesPushedEvent.load(std::memory_order_acquire);
 		if (!Event)
 		{
-			Event = FPlatformProcess::GetSynchEventFromPool(false);
-			SamplesPushedEvent.store(Event, std::memory_order_release);
+			FEvent* NewEvent = FPlatformProcess::GetSynchEventFromPool(false);
+			if (SamplesPushedEvent.compare_exchange_strong(Event, NewEvent, std::memory_order_release, std::memory_order_acquire))
+			{
+				// Use the new event.
+				Event = NewEvent;
+			}
+			else
+			{
+				// Use the existing event and return the new one.
+				FPlatformProcess::ReturnSynchEventToPool(NewEvent);
+			}
 		}
 
 		// Calculate when the timeout period will end.
