Fixed mismatched copy constructor/assignment operator in FOscFrequencyMod.

diff --git a/Engine/Source/Runtime/SignalProcessing/Public/DSP/Osc.h b/Engine/Source/Runtime/SignalProcessing/Public/DSP/Osc.h
--- a/Engine/Source/Runtime/SignalProcessing/Public/DSP/Osc.h
+++ b/Engine/Source/Runtime/SignalProcessing/Public/DSP/Osc.h
@@ -46,17 +46,6 @@ namespace Audio
 			, Semitones(0.0f)
 			, Cents(0.0f)
 		{}
-
-		FOscFrequencyMod(const FOscFrequencyMod& Other)
-			: Scale(Other.Scale)
-			, ExternalMod(Other.ExternalMod)
-			, Mod(Other.Mod)
-			, Detune(Other.Detune)
-			, PitchBend(Other.PitchBend)
-			, Octave(Other.Octave)
-			, Semitones(Other.Semitones)
-			, Cents(Other.Cents)
-		{}
 	};
 
 	// Oscillator base class
