Clamp input values to Dynamic Filter and SuperOscillator to prevent NaNs

diff --git a/Engine/Plugins/Runtime/Metasound/Source/MetasoundStandardNodes/Private/MetasoundSuperOscillatorNode.cpp b/Engine/Plugins/Runtime/Metasound/Source/MetasoundStandardNodes/Private/MetasoundSuperOscillatorNode.cpp
--- a/Engine/Plugins/Runtime/Metasound/Source/MetasoundStandardNodes/Private/MetasoundSuperOscillatorNode.cpp
+++ b/Engine/Plugins/Runtime/Metasound/Source/MetasoundStandardNodes/Private/MetasoundSuperOscillatorNode.cpp
@@ -1,5 +1,6 @@
 // Copyright Epic Games, Inc. All Rights Reserved.
 
+#include "AudioDefines.h"
 #include "Containers/Union.h"
 #include "Misc/TVariant.h"
 #include "DSP/Dsp.h"
@@ -167,7 +168,7 @@ namespace Metasound
 		METASOUND_PARAM(EntropyPin, "Entropy", "[0,1] Controls how evenly the voices are distributed in pitch")
 		METASOUND_PARAM(FrequencyModPin, "Modulation","Modulation Frequency Input (for doing FM)")
 		METASOUND_PARAM(GlideFactorPin, "Glide", "The amount of glide to use when changing frequencies. 0.0 = no glide, 1.0 = lots of glide.")
-		METASOUND_PARAM(MaxDetunePin, "Detune", "Max pitch offset of any Oscillators in Semitones. Only oscillators 2+ are detuned")
+		METASOUND_PARAM(MaxDetunePin, "Detune", "Max pitch offset of any Oscillators in Semitones, up to four octaves. Only oscillators 2+ are detuned")
 		METASOUND_PARAM(NumVoicesPin, "Voices", "[1,16] The number of Oscillators")
 		METASOUND_PARAM(PulseWidthPin, "Pulse Width", "The Width of the square part of the wave. Only used for square waves.")
 		METASOUND_PARAM(StereoWidthPin, "Width", "[0,1] Stereo Width of the oscillators")
@@ -285,8 +286,11 @@ namespace Metasound
 			const float ClampedFreq = FMath::Clamp(*BaseFrequency, -Nyquist, Nyquist);
 			const float ClampedEntropy = FMath::Clamp(*Entropy, 0.f, 1.f);
 			const float ClampedGlideEase = Audio::GetLogFrequencyClamped(*GlideFactor, { 0.0f, 1.0f }, { 1.0f, 0.0001f });
-			const float DetuneRatio = Audio::GetFrequencyMultiplier(*MaxDetune);
-			const float DetuneChannelGain = Audio::ConvertToLinear(*DetuneDb);
+
+			// allow +- 4 octaves of detune
+			constexpr float DetuneClamp = 12.f * 4.f;
+			const float DetuneRatio = Audio::GetFrequencyMultiplier(FMath::Clamp(*MaxDetune, -DetuneClamp, DetuneClamp));
+			const float DetuneChannelGain = Audio::ConvertToLinear(FMath::Clamp(*DetuneDb, MIN_VOLUME_DECIBELS, 24.f));
 			const float ClampedPulsueWidth = FMath::Clamp(*PulseWidth, 0.01f, 0.99f);
 			
 			const int32 NumFrames = ScratchBuffer.Num();
@@ -639,7 +643,7 @@ namespace Metasound
 			AudioRight->Zero();
 		}
 
-		virtual float GetStereoWidth() const override { return *StereoWidth; }
+		virtual float GetStereoWidth() const override { return FMath::Clamp(*StereoWidth, 0.f, 1.f); }
 		
 		virtual void FillOutputs(const float Gain = 1.f, const float LinearPan = 0.f) final override
 		{
diff --git a/Engine/Source/Runtime/SignalProcessing/Private/DynamicStateVariableFilter.cpp b/Engine/Source/Runtime/SignalProcessing/Private/DynamicStateVariableFilter.cpp
--- a/Engine/Source/Runtime/SignalProcessing/Private/DynamicStateVariableFilter.cpp
+++ b/Engine/Source/Runtime/SignalProcessing/Private/DynamicStateVariableFilter.cpp
@@ -1,9 +1,10 @@
 // Copyright Epic Games, Inc. All Rights Reserved.
 
 #include "DSP/DynamicStateVariableFilter.h"
+
+#include "AudioDefines.h"
 #include "DSP/BufferVectorOperations.h"
 #include "DSP/Dsp.h"
-#include "HAL/IConsoleManager.h"
 
 namespace Audio
 {
@@ -301,7 +302,7 @@ namespace Audio
 	{
 		if (FMath::IsNearlyEqual(Frequency, InFrequency) == false)
 		{
-			Frequency = FMath::Clamp(InFrequency, 5.0f, (SampleRate * 0.5f) - 1.0f);
+			Frequency = FMath::Clamp(InFrequency, MIN_FILTER_FREQUENCY, (SampleRate * 0.5f) - 1.0f);
 
 			KeyVars.G = FMath::Tan(PI * Frequency * OneOverSampleRate);
 
@@ -313,7 +314,8 @@ namespace Audio
 	{
 		if (FMath::IsNearlyEqual(InQ, Q) == false)
 		{
-			Q = FMath::Max(InQ, KINDA_SMALL_NUMBER);
+			constexpr float Max_Q = 50.f;
+			Q = FMath::Clamp(InQ, KINDA_SMALL_NUMBER, Max_Q);
 			OneOverQ = 1.f / Q;
 			bNeedsUpdate = true;
 		}
@@ -323,7 +325,7 @@ namespace Audio
 	{
 		if (FMath::IsNearlyEqual(GainDb, InGain) == false)
 		{
-			GainDb = InGain;
+			GainDb = FMath::Clamp(InGain, MIN_VOLUME_DECIBELS, -(MIN_VOLUME_DECIBELS));
 			bNeedsUpdate = true;
 		}
 	}
@@ -344,17 +346,17 @@ namespace Audio
 
 	void FDynamicStateVariableFilter::SetAttackTime(const float InAttackTime)
 	{
-		Envelope.SetAttackTime(InAttackTime);
+		Envelope.SetAttackTime(FMath::Clamp(InAttackTime, 0.f, UE_BIG_NUMBER));
 	}
 
 	void FDynamicStateVariableFilter::SetReleaseTime(const float InReleaseTime)
 	{
-		Envelope.SetReleaseTime(InReleaseTime);
+		Envelope.SetReleaseTime(FMath::Clamp(InReleaseTime, 0.f, UE_BIG_NUMBER));
 	}
 
 	void FDynamicStateVariableFilter::SetThreshold(const float InThresholdDb)
 	{
-		ThresholdDb = InThresholdDb;
+		ThresholdDb = FMath::Clamp(InThresholdDb, MIN_VOLUME_DECIBELS, 0.f);
 	}
 
 	void FDynamicStateVariableFilter::SetEnvMode(const EPeakMode::Type InMode)
@@ -364,18 +366,18 @@ namespace Audio
 
 	void FDynamicStateVariableFilter::SetDynamicRange(const float InDynamicRange)
 	{
-		DynamicRangeDb = InDynamicRange;
+		DynamicRangeDb = FMath::Clamp(InDynamicRange, MIN_VOLUME_DECIBELS, -(MIN_VOLUME_DECIBELS));
 	}
 
 	void FDynamicStateVariableFilter::SetRatio(const float InRatio)
 	{
-		Ratio = FMath::Max(InRatio, 1.f + KINDA_SMALL_NUMBER);
+		Ratio = FMath::Clamp(InRatio, 1.f + KINDA_SMALL_NUMBER, 100.f);
 		SlopeFactor = 1.0f - 1.0f / Ratio;
 	}
 
 	void FDynamicStateVariableFilter::SetKnee(const float InKnee)
 	{
-		Knee = InKnee;
+		Knee = FMath::Clamp(InKnee, 0.f, -(MIN_VOLUME_DECIBELS));
 		HalfKnee = Knee * 0.5f;
 	}
 
