PCM16 to float version of ArrayMixIn to avoid buffer scratch copy (potentially large for things like WaveTables or other raw PCM data manipulations)
#rnx

diff --git a/Engine/Source/Runtime/SignalProcessing/Private/FloatArrayMath.cpp b/Engine/Source/Runtime/SignalProcessing/Private/FloatArrayMath.cpp
--- a/Engine/Source/Runtime/SignalProcessing/Private/FloatArrayMath.cpp
+++ b/Engine/Source/Runtime/SignalProcessing/Private/FloatArrayMath.cpp
@@ -2137,6 +2137,50 @@ namespace Audio
 		}
 	}
 
+	void ArrayMixIn(TArrayView<const int16> InPcm16Buffer, TArrayView<float> BufferToSumTo, const float Gain)
+	{
+		CSV_SCOPED_TIMING_STAT(Audio_Dsp, ArrayMixIn);
+
+		checkf(InPcm16Buffer.Num() == BufferToSumTo.Num(), TEXT("Buffers must be equal size"));
+
+		const int32 Num = InPcm16Buffer.Num();
+
+		const int32 NumToSimd = Num & MathIntrinsics::SimdMask;
+		const int32 NumNotToSimd = Num & MathIntrinsics::NotSimdMask;
+
+		const int16* InputPtr = InPcm16Buffer.GetData();
+		float* OutPtr = BufferToSumTo.GetData();
+
+		const float ConversionValue = Gain / static_cast<float>(TNumericLimits<int16>::Max());
+
+		if (NumToSimd)
+		{
+			const VectorRegister4Float ConversionVector = VectorSetFloat1(ConversionValue);
+			AlignedFloat4 FloatArray(GlobalVectorConstants::FloatZero);
+
+			for (int32 i = 0; i < NumToSimd; i += AUDIO_NUM_FLOATS_PER_VECTOR_REGISTER)
+			{
+				FloatArray[0] = (float)InputPtr[i];
+				FloatArray[1] = (float)InputPtr[i + 1];
+				FloatArray[2] = (float)InputPtr[i + 2];
+				FloatArray[3] = (float)InputPtr[i + 3];
+
+				const VectorRegister4Float InVector = FloatArray.ToVectorRegister();
+				const VectorRegister4Float OutData = VectorLoad(&OutPtr[i]);
+				const VectorRegister4Float ScaledVector = VectorMultiplyAdd(InVector, ConversionVector, OutData);
+				VectorStore(ScaledVector, &OutPtr[i]);
+			}
+		}
+
+		if (NumNotToSimd)
+		{
+			for (int32 i = NumToSimd; i < Num; i++)
+			{
+				OutPtr[i] += (float)InputPtr[i] * ConversionValue;
+			}
+		}
+	}
+
 	void ArrayFloatToPcm16(TArrayView<const float> InView, TArrayView<int16> OutView)
 	{
 		CSV_SCOPED_TIMING_STAT(Audio_Dsp, ArrayFloatToPcm16);
diff --git a/Engine/Source/Runtime/SignalProcessing/Public/DSP/FloatArrayMath.h b/Engine/Source/Runtime/SignalProcessing/Public/DSP/FloatArrayMath.h
--- a/Engine/Source/Runtime/SignalProcessing/Public/DSP/FloatArrayMath.h
+++ b/Engine/Source/Runtime/SignalProcessing/Public/DSP/FloatArrayMath.h
@@ -222,6 +222,9 @@ namespace Audio
 	SIGNALPROCESSING_API void ArrayMixIn(TArrayView<const float> InFloatBuffer, TArrayView<float> BufferToSumTo, const float Gain);
 	SIGNALPROCESSING_API void ArrayMixIn(TArrayView<const float> InFloatBuffer, TArrayView<float> BufferToSumTo);
 
+	/** Takes buffer InPcm16Buffer, converts it to float and optionally multiplies it by Gain, and adds it to BufferToSumTo. */
+	SIGNALPROCESSING_API void ArrayMixIn(TArrayView<const int16> InPcm16Buffer, TArrayView<float> BufferToSumTo, const float Gain = 1.0f);
+
 	/** This version of ArrayMixIn will fade from StartGain to EndGain. */
 	SIGNALPROCESSING_API void ArrayMixIn(TArrayView<const float> InFloatBuffer, TArrayView<float> BufferToSumTo, const float StartGain, const float EndGain);
 
