midi scale to note array fix

diff --git a/Engine/Plugins/Runtime/Metasound/Source/MetasoundStandardNodes/Private/MetasoundMidiNoteQuantizerNode.cpp b/Engine/Plugins/Runtime/Metasound/Source/MetasoundStandardNodes/Private/MetasoundMidiNoteQuantizerNode.cpp
--- a/Engine/Plugins/Runtime/Metasound/Source/MetasoundStandardNodes/Private/MetasoundMidiNoteQuantizerNode.cpp
+++ b/Engine/Plugins/Runtime/Metasound/Source/MetasoundStandardNodes/Private/MetasoundMidiNoteQuantizerNode.cpp
@@ -28,8 +28,8 @@ namespace Metasound
 		// inputs
 		METASOUND_PARAM(ParamNoteIn, "Note In", "Midi Note to quantize");
 		METASOUND_PARAM(ParamRootNote, "Root Note", "Midi note to treat as the Root (0 = C, 1 = D, etc.).  Octave does not matter. Values < 0 will clamp to 0");
-		METASOUND_PARAM(ParamScaleDegrees, "Scale Degrees"
-			, "Set of notes in ascending order, represeting half steps starting at 0.0f meaning the Root Note. The highest value MUST be the same as 0.0 one octave higher: 12.0 in most cases, but it doesn't have to be (i.e. could be 24.0 if you want to define a 2 octave range)");
+		METASOUND_PARAM(ParamScaleDegrees, "Scale Degrees", "Set of notes in ascending order, represeting half steps starting at 0.0f meaning the Root Note. Scale degrees should be the notes in a scale that not including the octave. Can be the output of the Scale To Note Array node.");
+		METASOUND_PARAM(ParamScaleRange, "Scale Range In", "The number of semitones in the scale. E.g. a regular diatonic scale will be 12 semitones. Exotic scales could be something else.");
 
 		// outputs
 		METASOUND_PARAM(ParamNoteOutput, "Note Out", "Quantized Note");
@@ -54,6 +54,7 @@ namespace Metasound
 			, const FFloatReadRef& InMidiNoteIn
 			, const FFloatReadRef& InRootNote
 			, const FArrayScaleDegreeReadRef& InScale
+			, const FFloatReadRef& InScaleRange
 		);
 
 		// node interface
@@ -73,6 +74,7 @@ namespace Metasound
 		FFloatReadRef MidiNoteIn;
 		FFloatReadRef RootNote;
 		FArrayScaleDegreeReadRef Scale;
+		FFloatReadRef ScaleRange;
 
 		// output pins
 		FFloatWriteRef MidiNoteOut;
@@ -82,6 +84,7 @@ namespace Metasound
 		float PreviousRoot = -1.0f;
 		TArray<float> PreviousScale;
 		float PreviousNoteOut = 0.0f;
+		float PreviousScaleRange = 12.0f;
 
 		static constexpr float MaxNote = 1e12f;
 		static constexpr float MinNote = -1e12f;
@@ -101,10 +104,12 @@ namespace Metasound
 		, const FFloatReadRef& InMidiNoteIn
 		, const FFloatReadRef& InRootNote
 		, const FArrayScaleDegreeReadRef& InScale
-		)
+		, const FFloatReadRef& InScaleRange
+	)
 		: MidiNoteIn(InMidiNoteIn)
 		, RootNote(InRootNote)
 		, Scale(InScale)
+		, ScaleRange(InScaleRange)
 		, MidiNoteOut(FFloatWriteRef::CreateNew())
 	{
 		Reset(InParams);
@@ -142,6 +147,7 @@ namespace Metasound
 			FInputVertexInterface(
 				TInputDataVertex<float>(METASOUND_GET_PARAM_NAME_AND_METADATA(ParamNoteIn), 60.0f),
 				TInputDataVertex<float>(METASOUND_GET_PARAM_NAME_AND_METADATA(ParamRootNote), 0.0f),
+				TInputDataVertex<float>(METASOUND_GET_PARAM_NAME_AND_METADATA(ParamScaleRange), 12.0f),
 				TInputDataVertex<ScaleDegreeArrayType>(METASOUND_GET_PARAM_NAME_AND_METADATA(ParamScaleDegrees))
 			),
 			FOutputVertexInterface(
@@ -157,16 +163,20 @@ namespace Metasound
 	{
 		const FDataReferenceCollection& InputDataRefs = InParams.InputDataReferences;
 
+		const FInputVertexInterface& InputInterface = InParams.Node.GetVertexInterface().GetInputInterface();
+
 		// inputs
-		FFloatReadRef MidiNoteIn = InputDataRefs.GetDataReadReferenceOrConstruct<float>(METASOUND_GET_PARAM_NAME(ParamNoteIn));
-		FFloatReadRef RootNoteIn = InputDataRefs.GetDataReadReferenceOrConstruct<float>(METASOUND_GET_PARAM_NAME(ParamRootNote));
+		FFloatReadRef MidiNoteIn = InputDataRefs.GetDataReadReferenceOrConstructWithVertexDefault<float>(InputInterface, METASOUND_GET_PARAM_NAME(ParamNoteIn), InParams.OperatorSettings);
+		FFloatReadRef RootNoteIn = InputDataRefs.GetDataReadReferenceOrConstructWithVertexDefault<float>(InputInterface, METASOUND_GET_PARAM_NAME(ParamRootNote), InParams.OperatorSettings);
 		FArrayScaleDegreeReadRef InScaleArray = InParams.InputDataReferences.GetDataReadReferenceOrConstruct<ScaleDegreeArrayType>(METASOUND_GET_PARAM_NAME(ParamScaleDegrees));
+		FFloatReadRef ScaleRangeIn = InputDataRefs.GetDataReadReferenceOrConstructWithVertexDefault<float>(InputInterface, METASOUND_GET_PARAM_NAME(ParamScaleRange), InParams.OperatorSettings);
 
 		return MakeUnique <FMidiNoteQuantizerOperator>(
 			  InParams
 			, MidiNoteIn
 			, RootNoteIn
 			, InScaleArray
+			, ScaleRangeIn
 			);
 	}
 
@@ -174,6 +184,7 @@ namespace Metasound
 	{
 		InOutVertexData.BindReadVertex(METASOUND_GET_PARAM_NAME(ParamNoteIn), MidiNoteIn);
 		InOutVertexData.BindReadVertex(METASOUND_GET_PARAM_NAME(ParamRootNote), RootNote);
+		InOutVertexData.BindReadVertex(METASOUND_GET_PARAM_NAME(ParamScaleRange), ScaleRange);
 		InOutVertexData.BindReadVertex(METASOUND_GET_PARAM_NAME(ParamScaleDegrees), Scale);
 	}
 
@@ -203,9 +214,10 @@ namespace Metasound
 		PreviousNoteIn = FMath::Clamp(*MidiNoteIn, MinNote, MaxNote);
 		PreviousRoot = FMath::Clamp(*RootNote, MinRoot, MaxRoot);
 		PreviousScale = *Scale;
+		PreviousScaleRange = FMath::Max(*ScaleRange, 1.0f);
 
 		PreviousScale.Sort();
-		PreviousNoteOut = Audio::FMidiNoteQuantizer::QuantizeMidiNote(PreviousNoteIn, PreviousRoot, PreviousScale);
+		PreviousNoteOut = Audio::FMidiNoteQuantizer::QuantizeMidiNote(PreviousNoteIn, PreviousRoot, PreviousScale, PreviousScaleRange);
 		*MidiNoteOut = PreviousNoteOut;
 	}
 
@@ -220,18 +232,21 @@ namespace Metasound
 		// calculate new output and cache values if needed
 		const float CurrentNote = FMath::Clamp(*MidiNoteIn, MinNote, MaxNote);
 		const float CurrentRoot = FMath::Clamp(*RootNote, MinRoot, MaxRoot);
+		const float CurrentScaleRange = FMath::Max(*ScaleRange, 1.0f);;
 		if (!FMath::IsNearlyEqual(PreviousNoteIn, CurrentNote)
 			|| !FMath::IsNearlyEqual(PreviousRoot, CurrentRoot)
 			|| PreviousScale !=  *Scale
+			|| !FMath::IsNearlyEqual(PreviousScaleRange, CurrentScaleRange)
 			)
 		{
 			// cache values
 			PreviousNoteIn = CurrentNote;
 			PreviousRoot = CurrentRoot;
 			PreviousScale = *Scale;
+			PreviousScaleRange = CurrentScaleRange;
 
 			PreviousScale.Sort();
-			PreviousNoteOut = Audio::FMidiNoteQuantizer::QuantizeMidiNote(PreviousNoteIn, PreviousRoot, PreviousScale);
+			PreviousNoteOut = Audio::FMidiNoteQuantizer::QuantizeMidiNote(PreviousNoteIn, PreviousRoot, PreviousScale, PreviousScaleRange);
 		}
 
 		// set the output value
diff --git a/Engine/Plugins/Runtime/Metasound/Source/MetasoundStandardNodes/Private/MetasoundMidiScaleToArrayNode.cpp b/Engine/Plugins/Runtime/Metasound/Source/MetasoundStandardNodes/Private/MetasoundMidiScaleToArrayNode.cpp
--- a/Engine/Plugins/Runtime/Metasound/Source/MetasoundStandardNodes/Private/MetasoundMidiScaleToArrayNode.cpp
+++ b/Engine/Plugins/Runtime/Metasound/Source/MetasoundStandardNodes/Private/MetasoundMidiScaleToArrayNode.cpp
@@ -31,7 +31,6 @@ namespace Metasound
 	DEFINE_METASOUND_ENUM_ENTRY(Audio::EMusicalScale::Scale::Dominant7th_Mixolydian, "Dominant7th_MixolydianDescription", "Dominant 7th (Mixolydian)", "Dominant7th_MixolydianDescriptionTT", "Mioxlydian (Dominant 7)"),
 	DEFINE_METASOUND_ENUM_ENTRY(Audio::EMusicalScale::Scale::NaturalMinor_Aeolian, "NaturalMinor_AeolianDescription", "Natural Minor (Aeolian)", "NaturalMinor_AeolianDescriptionTT", "Natural Minor (Aeolian)"),
 	DEFINE_METASOUND_ENUM_ENTRY(Audio::EMusicalScale::Scale::HalfDiminished_Locrian, "HalfDiminished_LocrianDescription", "Half Diminished (Locrian)", "HalfDiminished_LocrianDescriptionTT", "Half-Diminished (Locrian)"),
-	DEFINE_METASOUND_ENUM_ENTRY(Audio::EMusicalScale::Scale::Diminished, "DiminishedDescription", "Diminished ", "DiminishedDescriptionTT", "Diminished"),
 	// non-diatonic
 	DEFINE_METASOUND_ENUM_ENTRY(Audio::EMusicalScale::Scale::Chromatic, "ChromaticDescription", "Chromatic", "ChromaticDescriptionTT", "Chromatic"),
 	DEFINE_METASOUND_ENUM_ENTRY(Audio::EMusicalScale::Scale::WholeTone, "WholeToneDescription", "Whole-Tone", "WholeToneDescriptionTT", "Whole Tone"),
@@ -55,6 +54,7 @@ namespace Metasound
 	DEFINE_METASOUND_ENUM_ENTRY(Audio::EMusicalScale::Scale::LydianDominant, "LydianDominantDescription", "Lydian Dominant ", "LydianDominantDescriptionTT", "Lydian Dominant"),
 	DEFINE_METASOUND_ENUM_ENTRY(Audio::EMusicalScale::Scale::Augmented, "AugmentedDescription", "Augmented", "AugmentedDescriptionTT", "Augmented"),
 	// diminished
+	DEFINE_METASOUND_ENUM_ENTRY(Audio::EMusicalScale::Scale::Diminished, "DiminishedDescription", "Diminished ", "DiminishedDescriptionTT", "Diminished"),
 	DEFINE_METASOUND_ENUM_ENTRY(Audio::EMusicalScale::Scale::Diminished_BeginWithHalfStep, "Diminished_BeginWithHalfStepDescription", "Diminished (Begin With Half-Step)", "Diminished_BeginWithHalfStepDescriptionTT", "Diminished (begins with Half Step)"),
 	DEFINE_METASOUND_ENUM_ENTRY(Audio::EMusicalScale::Scale::Diminished_BeginWithWholeStep, "Diminished_BeginWithWholeStepDescription", "Diminished (Begin With Whole-Step", "Diminished_BeginWithWholeStepDescriptionTT", "Diminished (begins with Whole Step)"),
 	DEFINE_METASOUND_ENUM_ENTRY(Audio::EMusicalScale::Scale::HalfDiminished_LocrianNumber2, "HalfDiminished_LocrianNumber2Description", "Half-Diminished (Locrian #2)", "HalfDiminished_LocrianNumber2DescriptionTT", "Half Diminished Locrian (#2)"),
@@ -68,10 +68,10 @@ namespace Metasound
 	{
 		// inputs
 		METASOUND_PARAM(ParamScaleDegreesPreset, "Scale Degrees", "Select scale preset");
-		METASOUND_PARAM(ParamChordTonesOnly, "Chord Tones Only", "If true, will only return a subset of the scale represeting chord tones. (i.e. scale degrees 1,3,5,7)");
+		METASOUND_PARAM(ParamChordTonesOnly, "Chord Tones Only", "If true, will only return a subset of the scale represeting chord tones. (i.e. scale degrees 1,3,5,7). Will not include chord extensions (i.e. 9, 11, 13).");
 
 		// outputs
-		METASOUND_PARAM(ParamNoteArrayOutput, "Scale Array Out", "Array represeting the scale as half steps above the root. The set is inclusive at both ends: (starting at 0.0f and ending with 12.0f)");
+		METASOUND_PARAM(ParamNoteArrayOutput, "Scale Array Out", "Array represeting the scale as semitones above the root starting at 0.0. The scale only includes one octave and each note only once.");
 
 	} // namespace MusicalScaleToNoteArrayParameterNames
 
@@ -169,8 +169,8 @@ namespace Metasound
 	{
 		static const FVertexInterface Interface(
 			FInputVertexInterface(
-				TInputDataVertex<FEnumEMusicalScale>(METASOUND_GET_PARAM_NAME_AND_METADATA(ParamScaleDegreesPreset)),
-				TInputDataVertex<bool>(METASOUND_GET_PARAM_NAME_AND_METADATA(ParamChordTonesOnly))
+				TInputDataVertex<FEnumEMusicalScale>(METASOUND_GET_PARAM_NAME_AND_METADATA(ParamScaleDegreesPreset), (int32)Audio::EMusicalScale::Scale::Major),
+				TInputDataVertex<bool>(METASOUND_GET_PARAM_NAME_AND_METADATA(ParamChordTonesOnly), false)
 			),
 			FOutputVertexInterface(
 				TOutputDataVertex<ScaleDegreeArrayType>(METASOUND_GET_PARAM_NAME_AND_METADATA(ParamNoteArrayOutput))
@@ -184,10 +184,11 @@ namespace Metasound
 	TUniquePtr<IOperator> FMusicalScaleToNoteArrayOperator::CreateOperator(const FCreateOperatorParams& InParams, FBuildErrorArray& OutErrors)
 	{
 		const FDataReferenceCollection& InputDataRefs = InParams.InputDataReferences;
+		const FInputVertexInterface& InputInterface = InParams.Node.GetVertexInterface().GetInputInterface();
 
 		// inputs
-		FEnumMusicalScaleReadRef Scale = InputDataRefs.GetDataReadReferenceOrConstruct<FEnumEMusicalScale>(METASOUND_GET_PARAM_NAME(ParamScaleDegreesPreset));
-		FBoolReadRef ChordTonesOnly = InputDataRefs.GetDataReadReferenceOrConstruct<bool>(METASOUND_GET_PARAM_NAME(ParamChordTonesOnly));
+		FEnumMusicalScaleReadRef Scale = InputDataRefs.GetDataReadReferenceOrConstructWithVertexDefault<FEnumEMusicalScale>(InputInterface, METASOUND_GET_PARAM_NAME(ParamScaleDegreesPreset), InParams.OperatorSettings);
+		FBoolReadRef ChordTonesOnly = InputDataRefs.GetDataReadReferenceOrConstructWithVertexDefault<bool>(InputInterface, METASOUND_GET_PARAM_NAME(ParamChordTonesOnly), InParams.OperatorSettings);
 
 		return MakeUnique <FMusicalScaleToNoteArrayOperator>(
 			  InParams
diff --git a/Engine/Source/Runtime/SignalProcessing/Private/MidiNoteQuantizer.cpp b/Engine/Source/Runtime/SignalProcessing/Private/MidiNoteQuantizer.cpp
--- a/Engine/Source/Runtime/SignalProcessing/Private/MidiNoteQuantizer.cpp
+++ b/Engine/Source/Runtime/SignalProcessing/Private/MidiNoteQuantizer.cpp
@@ -5,27 +5,21 @@
 
 namespace Audio
 {
-	// useful musical defines for scale degrees + accidentals to translate to half-steps
-	// i.e. - n1 = "Natural 1" = 0.0 half-steps from the root
-	//		- b3 = "Flat 3"    = 3.0f half-steps etc.
-	static const float n8 = 12.0f;  // C
-	static const float n7 = 11.0f;	// B
-	static const float		b7 = 10.0f;
-	static const float bb7 = 9.0f;
-	static const float n6 = 9.0f;	// A
-	static const float		b6 = 8.0f;
-	static const float		s5 = 8.0f;
-	static const float n5 = 7.0f;	// G
-	static const float		b5 = 6.0f;
-	static const float		s4 = 6.0f;
-	static const float n4 = 5.0f;	// F
-	static const float n3 = 4.0f;	// E
-	static const float		b3 = 3.0f;
-	static const float		s2 = 3.0f;
-	static const float n2 = 2.0f;	// D
-	static const float		b2 = 1.0f;
-	static const float n1 = 0.0f;	// C
-
+	namespace Note
+	{
+		static const float B1 = 11.0f;
+		static const float Bb1 = 10.0f;
+		static const float A1 = 9.0f;
+		static const float Ab1 = 8.0f;
+		static const float G1 = 7.0f;
+		static const float Gb1 = 6.0f;
+		static const float F1 = 5.0f;
+		static const float E1 = 4.0f;
+		static const float Eb1 = 3.0f;
+		static const float D1 = 2.0f;
+		static const float Db1 = 1.0f;
+		static const float C1 = 0.0f;
+	}
 
 	ScaleDegreeSet::ScaleDegreeSet(const TArray<float>& InScaleDegrees, const TArray<float>& InChordTones)
 		: ScaleDegrees(InScaleDegrees)
@@ -48,83 +42,44 @@ namespace Audio
 		return QuantizeMidiNote(InNote, InRoot, ScaleDegreeSetMap[InScale].GetScaleDegreeSet(bChordTonesOnlyIfApplicable));
 	}
 
-	float FMidiNoteQuantizer::QuantizeMidiNote(const float InNote, const float InRoot, const TArrayView<float> InScaleDegrees)
+	float FMidiNoteQuantizer::QuantizeMidiNote(const float InNote, const float InRoot, const TArrayView<float> InScaleDegrees, const float InSemitoneScaleRange)
 	{
-		// QuantizeValueToScaleDegree() works within a single octave, so compensate for the octave and root
-		const float SetRange = InScaleDegrees[InScaleDegrees.Num() - 1]; // (SetRange will == 12.0 for a normal octave)
+		// QuantizeValueToScaleDegree() works within a single octave, so map incoming note to within the first octave
 		const float NoteRootDelta = (InNote - InRoot);
-		const float InNoteOctave = FMath::FloorToFloat(NoteRootDelta / FMath::Max(SetRange, UE_SMALL_NUMBER));
-		const float ValueToQuant = NoteRootDelta - (SetRange * InNoteOctave);
+		const float ClampedSemitoneScaleRange = FMath::Max(InSemitoneScaleRange, 1.0f);
+		const float InNoteOctave = FMath::FloorToFloat(NoteRootDelta / ClampedSemitoneScaleRange);
+		const float ValueToQuant = NoteRootDelta - (InSemitoneScaleRange * InNoteOctave);
 
 		const float QuantizedValue = QuantizeValueToScaleDegree(ValueToQuant, InScaleDegrees);
 
-		// reconstruct the quantized midi note given the InRoot and original octave of InNote
-		return (SetRange * InNoteOctave + QuantizedValue) + InRoot;
+		// reconstruct the quantized midi note given the InRoot and original octave
+		return (InSemitoneScaleRange * InNoteOctave + QuantizedValue) + InRoot;
 	}
 
-	float FMidiNoteQuantizer::QuantizeValueToScaleDegree(const float InValue, const TArrayView<float> InScaleDegrees)
+	float FMidiNoteQuantizer::QuantizeValueToScaleDegree(const float InValue, const TArrayView<float> InScaleDegrees, const float InSemitoneScaleRange)
 	{
 		// the Set we are quantizing to should have at least 2 elements
-		if (!ensure(InScaleDegrees.Num()))
+		if (!InScaleDegrees.Num())
 		{
 			return InValue;
 		}
-		else if (InScaleDegrees.Num() == 1)
-		{
-			// if the set only contains a single value, that is technically the quantized value
-			return InScaleDegrees[0];
-		}
-		else if (!ensureMsgf((0.0f <= InValue) && (InValue <= InScaleDegrees[InScaleDegrees.Num() - 1])
-			, TEXT("Pitch Quantization: InValue (%f) is not in the range of a single octave (Needs to satisfy: 0.0f <= InValue <= %f)")
-			, InValue, InScaleDegrees[InScaleDegrees.Num() - 1]))
-		{
-			return InValue; // bail
-		}
 
-		bool bNegativeDelta = false;
-		float PositiveDelta = TNumericLimits<float>::Max();
-		float NegativeDelta = TNumericLimits<float>::Lowest();
+		// First check the min delta against the scale range
+		float QuantizedValue = InSemitoneScaleRange;
+		float CurrMinDelta = FMath::Abs(InValue - InSemitoneScaleRange);
 
-		// In the case where InValue == last element in the note set
-		// we want to initialize FirstNegativeIndex like this for our
-		// post-loop logic to return the correct value:
-		//	--> return Set[FirstNegativeIndex - 1];
-		// (in this case "NegativeDelta" will never get updated before the loop completes)
-		int32 FirstNegativeIndex = InScaleDegrees.Num();
+		// Then check the given scale degree array
 		for (int32 i = 0; i < InScaleDegrees.Num(); ++i)
 		{
-			// InValue == Set[i]
-			if (FMath::IsNearlyEqual(InValue, InScaleDegrees[i]))
-			{
-				return InScaleDegrees[i];
-			}
-
-			// We are measuring Delta as (InValue - InScaleDegrees[i]).
-			// Delta should be positive until we pass by InValue in the set.
-			// Once we hit a negative value, we know which two candidates we need
-			// to compare for the shortest distance, and can exit the loop
-			const float CurrDelta = InValue - InScaleDegrees[i];
-			if (CurrDelta >= 0.0f)
-			{
-				PositiveDelta = CurrDelta;
-				continue;
-			}
-			else
+			float NewDelta = FMath::Abs(InValue - InScaleDegrees[i]);
+			if (NewDelta < CurrMinDelta)
 			{
-				NegativeDelta = CurrDelta;
-				FirstNegativeIndex = i;
-				break;
+				CurrMinDelta = NewDelta;
+				QuantizedValue = InScaleDegrees[i];
 			}
 		}
 
-		if (FMath::Abs(NegativeDelta) < FMath::Abs(PositiveDelta))
-		{
-			return InScaleDegrees[FirstNegativeIndex];
-		}
-		else
-		{
-			return InScaleDegrees[FirstNegativeIndex - 1];
-		}
+		return QuantizedValue;
 	}
 
 
@@ -133,11 +88,11 @@ namespace Audio
 	/* REQUIREMENTS:
 	*	The (optional) second array defines Chord Tones, if empty, the first array will be used always
 	*		- this is ideally a subset of the first array (user-facing UI/UX will imply this), but does not technically need to be and is not enforced.
-	* 
+	*
 	*	The first array MUST have at least 1 entry or it will ensure.
-	* 
+	*
 	*	Both array MUST be in ascending order, or quantization behavior will be broken. This is not enforced at runtime and quantization and will not warn/ensure
-	* 
+	*
 	*	The Arrays must represent a full octave (inclusive): the lowest element being (0.0 and the highest being 12.0)
 	*/
 	TMap<EMusicalScale::Scale, ScaleDegreeSet> FMidiNoteQuantizer::ScaleDegreeSetMap{
@@ -146,64 +101,64 @@ namespace Audio
 			{{ [(required) SCALE DEGREES] }, { [(optional) CHORD TONES] }}}
 		*/
 		  { EMusicalScale::Scale::Major,
-			{{ n1, n2, n3, n4, n5, n6, n7, n8 }, { n1, n2, n3, n5, n7, n8 }}}
+			{{ Note::C1, Note::D1, Note::E1, Note::F1, Note::G1, Note::A1, Note::B1 }, { Note::C1, Note::E1, Note::G1, Note::B1 }}}
 		, { EMusicalScale::Scale::Dominant7th_Mixolydian,
-			{{ n1, n2, n3, n4, n5, n6, b7, n8 }, { n1, n2, n3, n5, b7, n8 }}}
+			{{ Note::C1, Note::D1, Note::E1, Note::F1, Note::G1, Note::A1, Note::Bb1 }, { Note::C1, Note::E1, Note::G1, Note::Bb1 }}}
 		, { EMusicalScale::Scale::Minor_Dorian,
-			{{ n1, n2, b3, n4, n5, n6, b7, n8 }, { n1, n2, b3, n5, b7, n8 }}}
+			{{ Note::C1, Note::D1, Note::Eb1, Note::F1, Note::G1, Note::A1, Note::Bb1 }, { Note::C1, Note::Eb1, Note::G1, Note::Bb1 }}}
 		, { EMusicalScale::Scale::HalfDiminished_Locrian,
-			{{ n1, b2, b3, n4, b5, b6, b7, n8 }, { n1, b3, b5, b7, n8 }}}
+			{{ Note::C1, Note::Db1, Note::Eb1, Note::F1, Note::Gb1, Note::Ab1, Note::Bb1 }, { Note::C1, Note::Eb1, Note::Gb1, Note::Bb1 }}}
 		, { EMusicalScale::Scale::Diminished,
-			{{ n1, n2, b3, n4, b5, b6, n6, n7, n8 }, { n1, b3, n5, bb7, n8 }}}
+			{{ Note::C1, Note::D1, Note::Eb1, Note::F1, Note::Gb1, Note::Ab1, Note::A1, Note::B1 }, { Note::C1, Note::Eb1, Note::G1, Note::A1 }}}
 		, { EMusicalScale::Scale::MajorPentatonic,
-			{{ n1, n2, n3, n5, n6, n8 }, { n1, n3, n5, n7, n8 }}}
+			{{ Note::C1, Note::D1, Note::E1, Note::G1, Note::A1}, { Note::C1, Note::E1, Note::G1 }}}
 		, { EMusicalScale::Scale::Lydian,
-			{{ n1, n2, n3, s4, n5, n6, n7, n8 }, { n1, n3, n5, n7, n8 }}}
+			{{ Note::C1, Note::D1, Note::E1, Note::Gb1, Note::G1, Note::A1, Note::B1}, { Note::C1, Note::E1, Note::G1, Note::B1}}}
 		, { EMusicalScale::Scale::Bebop_Major,
-			{{ n1, n2, n3, s4, n5, s5, n6, n7, n8 }, { n1, n3, n5, n7, n8 }}}
+			{{ Note::C1, Note::D1, Note::E1, Note::Gb1, Note::G1, Note::Ab1, Note::A1, Note::B1}, { Note::C1, Note::E1, Note::G1, Note::B1}}}
 		, { EMusicalScale::Scale::HarmonicMajor,
-			{{ n1, n2, n3, n4, n5, b6, n7, n8 }, { n1, n3, n5, n7, n8 }}}
+			{{ Note::C1, Note::D1, Note::E1, Note::F1, Note::G1, Note::Ab1, Note::B1}, { Note::C1, Note::E1, Note::G1, Note::B1}}}
 		, { EMusicalScale::Scale::LydianAugmented,
-			{{ n1, n2, n3, s4, s5, n6, n7, n8 }, { n1, n3, s5, n7, n8 }}}
+			{{ Note::C1, Note::D1, Note::E1, Note::Gb1, Note::Ab1, Note::A1, Note::B1}, { Note::C1, Note::E1, Note::Ab1, Note::B1}}}
 		, { EMusicalScale::Scale::Augmented,
-			{{ n1, s2, n3, n4, n5, b6, n7, n8 }, { n1, n3, n5, n7, n8 }}}
+			{{ Note::C1, Note::Eb1, Note::E1, Note::F1, Note::G1, Note::Ab1, Note::B1}, { Note::C1, Note::E1, Note::G1, Note::B1}}}
 		, { EMusicalScale::Scale::SixthModeOfHarmonicMinor,
-			{{ n1, s2, n3, s4, n5, n6, n7, n8 }, { n1, n3, n5, n7, n8 }}}
+			{{ Note::C1, Note::Eb1, Note::E1, Note::Gb1, Note::G1, Note::A1, Note::B1}, { Note::C1, Note::E1, Note::G1, Note::B1}}}
 		, { EMusicalScale::Scale::Diminished_BeginWithHalfStep,
-			{{ n1, b2, s2, n3, s4, n5, n6, b7, n8 }, { n1, n3, n5, b7, n8 }}}
+			{{ Note::C1, Note::Db1, Note::Eb1, Note::E1, Note::Gb1, Note::G1, Note::A1, Note::Bb1}, { Note::C1, Note::E1, Note::G1, Note::Bb1}}}
 		, { EMusicalScale::Scale::Blues,
-			{{ n1, b3, n4, s4, n5, b7, n8 }, { n1, b3, n5, b7, n8 }}}
+			{{ Note::C1, Note::Eb1, Note::F1, Note::Gb1, Note::G1, Note::Bb1}, { Note::C1, Note::Eb1, Note::G1, Note::Bb1}}}
 		, { EMusicalScale::Scale::Bebop_Dominant,
-			{{ n1, n2, n3, n4, n5, n6, b7, n7, n8 }, { n1, n3, n5, n7, n8 }}}
+			{{ Note::C1, Note::D1, Note::E1, Note::F1, Note::G1, Note::A1, Note::Bb1, Note::B1}, { Note::C1, Note::E1, Note::G1, Note::B1}}}
 		, { EMusicalScale::Scale::Spanish_or_Jewish,
-			{{ n1, b2, n3, n4, n5, b6, b7, n8 }, { n1, n3, n5, b7, n8 }}}
+			{{ Note::C1, Note::Db1, Note::E1, Note::F1, Note::G1, Note::Ab1, Note::Bb1}, { Note::C1, Note::E1, Note::G1, Note::Bb1}}}
 		, { EMusicalScale::Scale::LydianDominant,
-			{{ n1, n2, n3, s4, n5, n6, b7, n8 }, { n1, n3, n5, b7, n8 }}}
+			{{ Note::C1, Note::D1, Note::E1, Note::Gb1, Note::G1, Note::A1, Note::Bb1}, { Note::C1, Note::E1, Note::G1, Note::Bb1}}}
 		, { EMusicalScale::Scale::Hindu,
-			{{ n1, n2, n3, n4, n5, b6, b7, n8 }, { n1, n3, n5, b7, n8 }}}
+			{{ Note::C1, Note::D1, Note::E1, Note::F1, Note::G1, Note::Ab1, Note::Bb1}, { Note::C1, Note::E1, Note::G1, Note::Bb1}}}
 		, { EMusicalScale::Scale::WholeTone,
-			{{ n1, n2, n3, s4, s5, b7, n8 }, { n1, n3, n5, b7, n8 }}}
+			{{ Note::C1, Note::D1, Note::E1, Note::Gb1, Note::Ab1, Note::Bb1}, { Note::C1, Note::E1, Note::G1, Note::Bb1}}}
 		, { EMusicalScale::Scale::Chromatic,
-			{{ n1, b2, n2, b3, n3, n4, b5, n5, b6, n6, b7, n7, n8 }, { /* NA */ }}}
+			{{ Note::C1, Note::Db1, Note::D1, Note::Eb1, Note::E1, Note::F1, Note::Gb1, Note::G1, Note::Ab1, Note::A1, Note::Bb1, Note::B1}, {Note::C1, Note::D1, Note::E1, Note::Gb1, Note::Ab1, Note::Bb1 /* Every other note */}}}
 		, { EMusicalScale::Scale::DiminishedWholeTone,
-			{{ n1, b2, s2, n3, s4, n5, n6, b7, n8 }, { n1, n3, n5, b7, n8 }}}
+			{{ Note::C1, Note::Db1, Note::Eb1, Note::E1, Note::Gb1, Note::G1, Note::A1, Note::Bb1}, { Note::C1, Note::E1, Note::G1, Note::Bb1}}}
 		, { EMusicalScale::Scale::MinorPentatonic,
-			{{ n1, b3, n4, n5, b7, n8 }, { n1, b3, n5, b7, n8 }}}
+			{{ Note::C1, Note::Eb1, Note::F1, Note::G1, Note::Bb1}, { Note::C1, Note::Eb1, Note::G1, Note::Bb1}}}
 		, { EMusicalScale::Scale::Bebop_Minor,
-			{{ n1, n2, b3, n3, n4, n5, n6, b7, n8 }, { n1, b3, n5, b7, n8 }}}
+			{{ Note::C1, Note::D1, Note::Eb1, Note::E1, Note::F1, Note::G1, Note::A1, Note::Bb1}, { Note::C1, Note::Eb1, Note::G1, Note::Bb1}}}
 		, { EMusicalScale::Scale::MelodicMinor,
-			{{ n1, n2, b3, n4, n5, n6, n7, n8 }, { n1, b3, n5, n7, n8 }}}
+			{{ Note::C1, Note::D1, Note::Eb1, Note::F1, Note::G1, Note::A1, Note::B1}, { Note::C1, Note::Eb1, Note::G1, Note::B1}}}
 		, { EMusicalScale::Scale::Bebop_MinorNumber2,
-			{{ n1, n2, b3, n4, n5, s5, n6, n7, n8 }, { n1, b3, n5, n7, n8 }}}
+			{{ Note::C1, Note::D1, Note::Eb1, Note::F1, Note::G1, Note::Ab1, Note::A1, Note::B1}, { Note::C1, Note::Eb1, Note::G1, Note::B1}}}
 		, { EMusicalScale::Scale::HarmonicMinor,
-			{{ n1, n2, b3, n4, n5, b6, n7, n8 }, { n1, b3, n5, n7, n8 }}}
+			{{ Note::C1, Note::D1, Note::Eb1, Note::F1, Note::G1, Note::Ab1, Note::B1}, { Note::C1, Note::Eb1, Note::G1, Note::B1}}}
 		, { EMusicalScale::Scale::Diminished_BeginWithWholeStep,
-			{{ n1, n2, b3, n4, s4, s5, n6, n7, n8 }, { n1, b3, s5, n7, n8 }}}
+			{{ Note::C1, Note::D1, Note::Eb1, Note::F1, Note::Gb1, Note::Ab1, Note::A1, Note::B1}, { Note::C1, Note::Eb1, Note::Ab1, Note::B1}}}
 		, { EMusicalScale::Scale::Phrygian,
-			{{ n1, b2, b3, n4, n5, b6, b7, n8 }, { n1, b3, n5, b7, n8 }}}
+			{{ Note::C1, Note::Db1, Note::Eb1, Note::F1, Note::G1, Note::Ab1, Note::Bb1}, { Note::C1, Note::Eb1, Note::G1, Note::Bb1}}}
 		, { EMusicalScale::Scale::NaturalMinor_Aeolian,
-			{{ n1, n2, b3, n4, n5, b6, b7, n8 }, { n1, b3, n5, b7, n8 }}}
+			{{ Note::C1, Note::D1, Note::Eb1, Note::F1, Note::G1, Note::Ab1, Note::Bb1}, { Note::C1, Note::Eb1, Note::G1, Note::Bb1}}}
 		, { EMusicalScale::Scale::HalfDiminished_LocrianNumber2,
-			{{ n1, b2, b3, n4, b5, b6, b7, n8 }, { n1, b3, b5, b7, n8 }}}
+			{{ Note::C1, Note::Db1, Note::Eb1, Note::F1, Note::Gb1, Note::Ab1, Note::Bb1}, { Note::C1, Note::Eb1, Note::Gb1, Note::Bb1}}}
 	};
 } // namespace Audio
diff --git a/Engine/Source/Runtime/SignalProcessing/Public/DSP/MidiNoteQuantizer.h b/Engine/Source/Runtime/SignalProcessing/Public/DSP/MidiNoteQuantizer.h
--- a/Engine/Source/Runtime/SignalProcessing/Public/DSP/MidiNoteQuantizer.h
+++ b/Engine/Source/Runtime/SignalProcessing/Public/DSP/MidiNoteQuantizer.h
@@ -12,37 +12,52 @@ namespace Audio
 		// in the "#pragma region Metasound Enums" area to show up as options in that Node
 		enum Scale
 		{
+			// Modes
 			Major,
-			Dominant7th_Mixolydian,
 			Minor_Dorian,
-			HalfDiminished_Locrian,
-			Diminished,
-			MajorPentatonic,
+			Phrygian,
 			Lydian,
-			Bebop_Major,
-			HarmonicMajor,
-			LydianAugmented,
-			Augmented,
-			SixthModeOfHarmonicMinor,
-			Diminished_BeginWithHalfStep,
-			Blues,
-			Bebop_Dominant,
-			Spanish_or_Jewish,
-			LydianDominant,
-			Hindu,
-			WholeTone,
+			Dominant7th_Mixolydian,
+			NaturalMinor_Aeolian,
+			HalfDiminished_Locrian,
+
+			// Non-Diatonic
 			Chromatic,
+			WholeTone,
 			DiminishedWholeTone,
+
+			// pentatonic
+			MajorPentatonic,
 			MinorPentatonic,
+			Blues,
+
+			// bebop
+			Bebop_Major,
 			Bebop_Minor,
-			MelodicMinor,
 			Bebop_MinorNumber2,
+			Bebop_Dominant,
+
+			// common major/minors
+			HarmonicMajor,
 			HarmonicMinor,
+			MelodicMinor,
+			SixthModeOfHarmonicMinor,
+
+			// lydian/augmented
+			LydianAugmented,
+			LydianDominant,
+			Augmented,
+
+			// diminished
+			Diminished,
+			Diminished_BeginWithHalfStep,
 			Diminished_BeginWithWholeStep,
-			Phrygian,
-			NaturalMinor_Aeolian,
 			HalfDiminished_LocrianNumber2,
 
+			// other
+			Spanish_or_Jewish,
+			Hindu,
+
 			Count
 		};
 	}
@@ -64,10 +79,10 @@ namespace Audio
 	{
 	public:
 		static float QuantizeMidiNote(const float InNote, const float InRoot, EMusicalScale::Scale InScale, bool bChordTonesOnlyIfApplicable = false);
-		static float QuantizeMidiNote(const float InNote, const float InRoot, const TArrayView<float> InScaleDegrees);
+		static float QuantizeMidiNote(const float InNote, const float InRoot, const TArrayView<float> InScaleDegrees, const float InSemitoneScaleRange = 12.0f);
 
 	private:
-		static float QuantizeValueToScaleDegree(const float InValue, const TArrayView<float> InScaleDegrees);
+		static float QuantizeValueToScaleDegree(const float InValue, const TArrayView<float> InScaleDegrees, const float InSemitoneScaleRange = 12.0f);
 
 	public:
 		// Statically-defined scale/chord tone definitions.
