[Engine]

diff --git a/Engine/Source/Runtime/SignalProcessing/Public/DSP/EQ.h b/Engine/Source/Runtime/SignalProcessing/Public/DSP/EQ.h
--- a/Engine/Source/Runtime/SignalProcessing/Public/DSP/EQ.h
+++ b/Engine/Source/Runtime/SignalProcessing/Public/DSP/EQ.h
@@ -10,35 +10,35 @@ namespace Audio
 	// Equalizer filter
 	// An equalizer is a cascaded (serial) band of parametric EQs
 	// This filter allows for setting each band with variable Bandwidth/Q, Frequency, and Gain
-	class SIGNALPROCESSING_API FEqualizer
+	class FEqualizer
 	{
 	public:
 		// Constructor
-		FEqualizer();
+		SIGNALPROCESSING_API FEqualizer();
 
 		// Destructor
-		~FEqualizer();
+		SIGNALPROCESSING_API ~FEqualizer();
 
 		// Initialize the equalizer
-		void Init(const float InSampleRate, const int32 InNumBands, const int32 InNumChannels);
+		SIGNALPROCESSING_API void Init(const float InSampleRate, const int32 InNumBands, const int32 InNumChannels);
 
 		// Sets whether or not the band is enabled
-		void SetBandEnabled(const int32 InBand, const bool bEnabled);
+		SIGNALPROCESSING_API void SetBandEnabled(const int32 InBand, const bool bEnabled);
 
 		// Sets all params of the band at once
-		void SetBandParams(const int32 InBand, const float InFrequency, const float InBandwidth, const float InGainDB);
+		SIGNALPROCESSING_API void SetBandParams(const int32 InBand, const float InFrequency, const float InBandwidth, const float InGainDB);
 
 		// Sets the band frequency
-		void SetBandFrequency(const int32 InBand, const float InFrequency);
+		SIGNALPROCESSING_API void SetBandFrequency(const int32 InBand, const float InFrequency);
 
 		// Sets the band resonance (use alternatively to bandwidth)
-		void SetBandBandwidth(const int32 InBand, const float InBandwidth);
+		SIGNALPROCESSING_API void SetBandBandwidth(const int32 InBand, const float InBandwidth);
 
 		// Sets the band gain in decibels
-		void SetBandGainDB(const int32 InBand, const float InGainDB);
+		SIGNALPROCESSING_API void SetBandGainDB(const int32 InBand, const float InGainDB);
 
 		// Processes the audio frame (audio frame must have channels equal to that used during initialization)
-		void ProcessAudioFrame(const float* InAudio, float* OutAudio);
+		SIGNALPROCESSING_API void ProcessAudioFrame(const float* InAudio, float* OutAudio);
 
 	private:
 
diff --git a/Engine/Source/Runtime/SignalProcessing/Public/DSP/EarlyReflectionsFast.h b/Engine/Source/Runtime/SignalProcessing/Public/DSP/EarlyReflectionsFast.h
--- a/Engine/Source/Runtime/SignalProcessing/Public/DSP/EarlyReflectionsFast.h
+++ b/Engine/Source/Runtime/SignalProcessing/Public/DSP/EarlyReflectionsFast.h
@@ -11,7 +11,7 @@
 
 namespace Audio
 {
-	struct SIGNALPROCESSING_API FEarlyReflectionsFastSettings
+	struct FEarlyReflectionsFastSettings
 	{
 		// Early reflections gain
 		float Gain;
@@ -28,49 +28,49 @@ namespace Audio
 		// Early reflection high frequency absorption factor
 		float Absorption;
 
-		FEarlyReflectionsFastSettings();
+		SIGNALPROCESSING_API FEarlyReflectionsFastSettings();
 
-		bool operator==(const FEarlyReflectionsFastSettings& Other) const;
+		SIGNALPROCESSING_API bool operator==(const FEarlyReflectionsFastSettings& Other) const;
 
-		bool operator!=(const FEarlyReflectionsFastSettings& Other) const;
+		SIGNALPROCESSING_API bool operator!=(const FEarlyReflectionsFastSettings& Other) const;
 	};
 
 	// Basic implementation of early reflections using a predelay, low pass filter and feedback delay network (FDN). The FDN
 	// utilizes four delay lines where each delay line consists of an all pass filter and low pass filter to control diffusion
 	// and absorption.
-	class SIGNALPROCESSING_API FEarlyReflectionsFast
+	class FEarlyReflectionsFast
 	{
 	public:
 		
 		// Limits for settings
-		static const float MaxGain;
-		static const float MinGain;
-		static const float MaxPreDelay;
-		static const float MinPreDelay;
-		static const float MaxBandwidth;
-		static const float MinBandwidth;
-		static const float MaxDecay;
-		static const float MinDecay;
-		static const float MaxAbsorption;
-		static const float MinAbsorption;
+		static SIGNALPROCESSING_API const float MaxGain;
+		static SIGNALPROCESSING_API const float MinGain;
+		static SIGNALPROCESSING_API const float MaxPreDelay;
+		static SIGNALPROCESSING_API const float MinPreDelay;
+		static SIGNALPROCESSING_API const float MaxBandwidth;
+		static SIGNALPROCESSING_API const float MinBandwidth;
+		static SIGNALPROCESSING_API const float MaxDecay;
+		static SIGNALPROCESSING_API const float MinDecay;
+		static SIGNALPROCESSING_API const float MaxAbsorption;
+		static SIGNALPROCESSING_API const float MinAbsorption;
 		
-		static const FEarlyReflectionsFastSettings DefaultSettings;
+		static SIGNALPROCESSING_API const FEarlyReflectionsFastSettings DefaultSettings;
 		
 		// InMaxNumInternalBufferSamples sets the maximum possible samples in an internal buffer.
-		FEarlyReflectionsFast(float InSampleRate, int32 InMaxNumInternalBufferSamples, const FEarlyReflectionsFastSettings& InSettings=DefaultSettings);
-		~FEarlyReflectionsFast();
+		SIGNALPROCESSING_API FEarlyReflectionsFast(float InSampleRate, int32 InMaxNumInternalBufferSamples, const FEarlyReflectionsFastSettings& InSettings=DefaultSettings);
+		SIGNALPROCESSING_API ~FEarlyReflectionsFast();
 
 		// Sets the reverb settings, clamps, applies, and updates
-		void SetSettings(const FEarlyReflectionsFastSettings& InSettings);
+		SIGNALPROCESSING_API void SetSettings(const FEarlyReflectionsFastSettings& InSettings);
 
 		// Process the single audio frame
-		void ProcessAudio(const FAlignedFloatBuffer& InSamples, const int32 InNumChannels, FAlignedFloatBuffer& OutLeftSamples, FAlignedFloatBuffer& OutRightSamples);
+		SIGNALPROCESSING_API void ProcessAudio(const FAlignedFloatBuffer& InSamples, const int32 InNumChannels, FAlignedFloatBuffer& OutLeftSamples, FAlignedFloatBuffer& OutRightSamples);
 
 		// Silence internal audio.
-		void FlushAudio();
+		SIGNALPROCESSING_API void FlushAudio();
 
 		// Clamps settings to acceptable values. 
-		static void ClampSettings(FEarlyReflectionsFastSettings& InOutSettings);
+		static SIGNALPROCESSING_API void ClampSettings(FEarlyReflectionsFastSettings& InOutSettings);
 
 	private:
 		void ApplySettings();
diff --git a/Engine/Source/Runtime/SignalProcessing/Public/DSP/Encoders/IAudioEncoder.h b/Engine/Source/Runtime/SignalProcessing/Public/DSP/Encoders/IAudioEncoder.h
--- a/Engine/Source/Runtime/SignalProcessing/Public/DSP/Encoders/IAudioEncoder.h
+++ b/Engine/Source/Runtime/SignalProcessing/Public/DSP/Encoders/IAudioEncoder.h
@@ -12,7 +12,7 @@ namespace Audio
 	* IAudioEncoder
 	* Interface used to implement a runtime audio encoder.
 	*/
-	class SIGNALPROCESSING_API IAudioEncoder
+	class IAudioEncoder
 	{
 	public:
 		/**
@@ -20,8 +20,8 @@ namespace Audio
 		 * Optionally, DataBufferSlack can be called with a different value
 		 * depending on how often PushAudio, EncodeIfPossible, or PopData will be called.
 		 */
-		IAudioEncoder(uint32 AudioBufferSlack, uint32 DataBufferSlack = 4096);
-		virtual ~IAudioEncoder();
+		SIGNALPROCESSING_API IAudioEncoder(uint32 AudioBufferSlack, uint32 DataBufferSlack = 4096);
+		SIGNALPROCESSING_API virtual ~IAudioEncoder();
 
 		/**
 		 * Call this function when audio is available. 
@@ -29,14 +29,14 @@ namespace Audio
 		 * Call this function with bEncodeIfPossible set to false.
 		 * Returns true if audio was successfully pushed, false, if the internal buffer is full.
 		 */
-		bool PushAudio(const float* InBuffer, int32 NumSamples, bool bEncodeIfPossible = true);
+		SIGNALPROCESSING_API bool PushAudio(const float* InBuffer, int32 NumSamples, bool bEncodeIfPossible = true);
 
 		/**
 		 * Pop compressed data. If you are using this encoder for streaming over network,
 		 * use the size returned by GetCompressedPacketSize().
 		 * Returns number of bytes written.
 		 */
-		int32 PopData(uint8* OutData, int32 DataSize);
+		SIGNALPROCESSING_API int32 PopData(uint8* OutData, int32 DataSize);
 
 		/**
 		 * Used for internet streaming. Should return the amount of bytes required for a self contained packet.
@@ -47,20 +47,20 @@ namespace Audio
 		 * If you'd like to run audio encoding on a separate thread, use this call.
 		 * Otherwise, ensure that bEncodeIfPossible is set to true when you call PushAudio.
 		 */
-		bool EncodeIfPossible();
+		SIGNALPROCESSING_API bool EncodeIfPossible();
 
 		/**
 		 * Call this once you are finished pushing audio.
 		 * Returns the amount of bytes left to pop if positive, and if negative indicates a failiure.
 		 */
-		int64 Finalize();
+		SIGNALPROCESSING_API int64 Finalize();
 
 	protected:
 		/**
 		 * Should be called in the constructor of any implementation of IAudioEncoder.
 		 * Calls StartFile if necessary.
 		 */
-		void Init(const FSoundQualityInfo& InQualityInfo);
+		SIGNALPROCESSING_API void Init(const FSoundQualityInfo& InQualityInfo);
 
 		/**
 		 * How many samples of decoded audio that are required for a single compression operation.
@@ -87,7 +87,7 @@ namespace Audio
 		virtual bool EndFile(TArray<uint8>& OutBytes) = 0;
 
 	private:
-		IAudioEncoder();
+		SIGNALPROCESSING_API IAudioEncoder();
 
 		TArray<float> CurrentAudioBuffer;
 		TArray<uint8> CurrentCompressedBuffer;
diff --git a/Engine/Source/Runtime/SignalProcessing/Public/DSP/Encoders/OggVorbisEncoder.h b/Engine/Source/Runtime/SignalProcessing/Public/DSP/Encoders/OggVorbisEncoder.h
--- a/Engine/Source/Runtime/SignalProcessing/Public/DSP/Encoders/OggVorbisEncoder.h
+++ b/Engine/Source/Runtime/SignalProcessing/Public/DSP/Encoders/OggVorbisEncoder.h
@@ -9,21 +9,21 @@
 
 struct FOggVorbisEncoderPrivateState;
 
-class SIGNALPROCESSING_API FOggVorbisEncoder : public Audio::IAudioEncoder
+class FOggVorbisEncoder : public Audio::IAudioEncoder
 {
 public:
-	FOggVorbisEncoder(const FSoundQualityInfo& InInfo, int32 AverageBufferCallbackSize);
+	SIGNALPROCESSING_API FOggVorbisEncoder(const FSoundQualityInfo& InInfo, int32 AverageBufferCallbackSize);
 
 	// From IAudioEncoder: returns 0, since Ogg Vorbis is not built for networked streaming.
-	virtual int32 GetCompressedPacketSize() const override;
+	SIGNALPROCESSING_API virtual int32 GetCompressedPacketSize() const override;
 
 protected:
 
 	// From IAudioEncoder:
-	virtual int64 SamplesRequiredPerEncode() const override;
-	virtual bool StartFile(const FSoundQualityInfo& InQualityInfo, TArray<uint8>& OutFileStart) override;
-	virtual bool EncodeChunk(const TArray<float>& InAudio, TArray<uint8>& OutBytes) override;
-	virtual bool EndFile(TArray<uint8>& OutBytes) override;
+	SIGNALPROCESSING_API virtual int64 SamplesRequiredPerEncode() const override;
+	SIGNALPROCESSING_API virtual bool StartFile(const FSoundQualityInfo& InQualityInfo, TArray<uint8>& OutFileStart) override;
+	SIGNALPROCESSING_API virtual bool EncodeChunk(const TArray<float>& InAudio, TArray<uint8>& OutBytes) override;
+	SIGNALPROCESSING_API virtual bool EndFile(TArray<uint8>& OutBytes) override;
 
 private:
 	FOggVorbisEncoder();
diff --git a/Engine/Source/Runtime/SignalProcessing/Public/DSP/Envelope.h b/Engine/Source/Runtime/SignalProcessing/Public/DSP/Envelope.h
--- a/Engine/Source/Runtime/SignalProcessing/Public/DSP/Envelope.h
+++ b/Engine/Source/Runtime/SignalProcessing/Public/DSP/Envelope.h
@@ -10,7 +10,7 @@
 namespace Audio
 {
 	// Envelope class generates ADSR style envelope
-	class SIGNALPROCESSING_API FEnvelope
+	class FEnvelope
 	{
 	public:
 		// States for the envelope state machine
@@ -24,14 +24,14 @@ namespace Audio
 			Shutdown
 		};
 		
-		FEnvelope();
-		virtual ~FEnvelope();
+		SIGNALPROCESSING_API FEnvelope();
+		SIGNALPROCESSING_API virtual ~FEnvelope();
 
 		// Initialize the envelope with the given sample rate
-		void Init(const float InSampleRate, const int32 InVoiceId = 0, FModulationMatrix* InModMatrix = nullptr, const bool bInSimulateAnalog = true);
+		SIGNALPROCESSING_API void Init(const float InSampleRate, const int32 InVoiceId = 0, FModulationMatrix* InModMatrix = nullptr, const bool bInSimulateAnalog = true);
 
 		// Sets the envelope mode 
-		void SetSimulateAnalog(const bool bInSimulatingAnalog);
+		SIGNALPROCESSING_API void SetSimulateAnalog(const bool bInSimulatingAnalog);
 
 		// Sets whether the envelope is in legato mode. Legato mode doesn't restart the envelope if it's already playing
 		void SetLegato(const bool bInLegatoMode) { bIsLegatoMode = bInLegatoMode; }
@@ -41,59 +41,59 @@ namespace Audio
 		bool IsRetrigger() const { return bIsRetriggerMode; }
 
 		// Start the envelope, puts envelope in attack state
-		virtual void Start();
+		SIGNALPROCESSING_API virtual void Start();
 
 		// For truly legato envelope logic, we need to know the new sustain gain (if its being changed).
-		virtual void StartLegato(const float InNewDepth);
+		SIGNALPROCESSING_API virtual void StartLegato(const float InNewDepth);
 		virtual void StartLegato() { StartLegato(Depth); }
 
 		// Stop the envelope, puts in the release state. Can optionally force to off state.
-		virtual void Stop();
+		SIGNALPROCESSING_API virtual void Stop();
 
 		// Puts envelope into shutdown mode, which is a much faster cutoff than release, but avoids pops
-		virtual void Shutdown();
+		SIGNALPROCESSING_API virtual void Shutdown();
 
 		// Kills the envelope (will cause discontinuity)
-		virtual void Kill();
+		SIGNALPROCESSING_API virtual void Kill();
 
 		// Queries if the envelope has finished
-		virtual bool IsDone() const;
+		SIGNALPROCESSING_API virtual bool IsDone() const;
 
 		// Resets the envelope
-		virtual void Reset();
+		SIGNALPROCESSING_API virtual void Reset();
 
 		// Update the state of the envelope
-		virtual void Update();
+		SIGNALPROCESSING_API virtual void Update();
 
 		// Generate the next output value of the envelope.
 		// Optionally outputs the bias output (i.e. -1.0 to 1.0)
-		virtual float Generate(float* BiasedOutput = nullptr);
+		SIGNALPROCESSING_API virtual float Generate(float* BiasedOutput = nullptr);
 
-		virtual EEnvelopeState GetState() const;
+		SIGNALPROCESSING_API virtual EEnvelopeState GetState() const;
 
 		// Sets the envelope attack time in msec
-		virtual void SetAttackTime(const float InAttackTimeMsec);
+		SIGNALPROCESSING_API virtual void SetAttackTime(const float InAttackTimeMsec);
 
 		// Sets the envelope decay time in msec
-		virtual void SetDecayTime(const float InDecayTimeMsec);
+		SIGNALPROCESSING_API virtual void SetDecayTime(const float InDecayTimeMsec);
 
 		// Sets the envelope sustain gain in linear gain values
-		virtual void SetSustainGain(const float InSustainGain);
+		SIGNALPROCESSING_API virtual void SetSustainGain(const float InSustainGain);
 
 		// Sets the envelope release time in msec
-		virtual void SetReleaseTime(const float InReleaseTimeMsec);
+		SIGNALPROCESSING_API virtual void SetReleaseTime(const float InReleaseTimeMsec);
 
 		// Inverts the value of envelope output
-		virtual void SetInvert(const bool bInInvert);
+		SIGNALPROCESSING_API virtual void SetInvert(const bool bInInvert);
 
 		// Inverts the value of the biased envelope output
-		virtual void SetBiasInvert(const bool bInBiasInvert);
+		SIGNALPROCESSING_API virtual void SetBiasInvert(const bool bInBiasInvert);
 
 		// Sets the envelope depth.
-		virtual void SetDepth(const float InDepth);
+		SIGNALPROCESSING_API virtual void SetDepth(const float InDepth);
 
 		// Sets the depth of the bias output. 
-		virtual void SetBiasDepth(const float InDepth);
+		SIGNALPROCESSING_API virtual void SetBiasDepth(const float InDepth);
 
 		// Get the envelope's patch nodes
 		const FPatchSource GetModSourceEnv() const { return EnvSource; }
@@ -166,31 +166,31 @@ namespace Audio
 	};
 
 	// sample accurate attack-decay style envelope generator
-	class SIGNALPROCESSING_API FADEnvelope
+	class FADEnvelope
 	{
 	public:
 		FADEnvelope() {}
 		~FADEnvelope() {}
 
-		void Init(int32 InSampleRate);
+		SIGNALPROCESSING_API void Init(int32 InSampleRate);
 
-		void SetAttackTimeSeconds(float InAttackTimeSeconds);
-		void SetDecayTimeSeconds(float InReleaseTimeSeconds);
-		void SetAttackCurveFactor(float InAttackCurve);
-		void SetDecayCurveFactor(float InDecayCurve);
+		SIGNALPROCESSING_API void SetAttackTimeSeconds(float InAttackTimeSeconds);
+		SIGNALPROCESSING_API void SetDecayTimeSeconds(float InReleaseTimeSeconds);
+		SIGNALPROCESSING_API void SetAttackCurveFactor(float InAttackCurve);
+		SIGNALPROCESSING_API void SetDecayCurveFactor(float InDecayCurve);
 
 		void SetLooping(bool bInIsLooping) { bIsLooping = bInIsLooping; }
 		bool IsLooping() const { return bIsLooping; }
 
 		// Call function to trigger a new attack-phase of the envelope generator
-		void Attack();
+		SIGNALPROCESSING_API void Attack();
 
 		// Generates an output audio buffer (used for audio-rate envelopes)
-		void GetNextEnvelopeOut(int32 StartFrame, int32 EndFrame, TArray<int32>& OutFinishedFrames, Audio::AlignedFloatBuffer& OutEnvelope);
+		SIGNALPROCESSING_API void GetNextEnvelopeOut(int32 StartFrame, int32 EndFrame, TArray<int32>& OutFinishedFrames, Audio::AlignedFloatBuffer& OutEnvelope);
 
 		// Generates a single float value (used for control-rate envelopes)
-		void GetNextEnvelopeOut(int32 StartFrame, int32 EndFrame, TArray<int32>& OutFinishedFrames, float& OutEnvelope);
-		bool GetNextEnvelopeOut(float& OutEnvelope);
+		SIGNALPROCESSING_API void GetNextEnvelopeOut(int32 StartFrame, int32 EndFrame, TArray<int32>& OutFinishedFrames, float& OutEnvelope);
+		SIGNALPROCESSING_API bool GetNextEnvelopeOut(float& OutEnvelope);
 
 	private:
 		int32 CurrentSampleIndex = INDEX_NONE;
diff --git a/Engine/Source/Runtime/SignalProcessing/Public/DSP/EnvelopeFollower.h b/Engine/Source/Runtime/SignalProcessing/Public/DSP/EnvelopeFollower.h
--- a/Engine/Source/Runtime/SignalProcessing/Public/DSP/EnvelopeFollower.h
+++ b/Engine/Source/Runtime/SignalProcessing/Public/DSP/EnvelopeFollower.h
@@ -28,7 +28,7 @@ namespace Audio
 
 
 	/** Conversion between attack/release time and attack/release sample counts. */
-	class SIGNALPROCESSING_API FAttackRelease
+	class FAttackRelease
 	{
 		// see https://en.wikipedia.org/wiki/RC_time_constant
 		// Time constants indicate how quickly the envelope follower responds to changes in input
@@ -43,11 +43,11 @@ namespace Audio
 		 * @param InReleaseTimeMsec - The desired release time in milliseconds.
 		 * @param bInIsAnalog - Whether to model analog RC circuits or use digital models.
 		 */
-		FAttackRelease(float InSampleRate, float InAttackTimeMsec, float InReleaseTimeMsec, bool bInIsAnalog);
+		SIGNALPROCESSING_API FAttackRelease(float InSampleRate, float InAttackTimeMsec, float InReleaseTimeMsec, bool bInIsAnalog);
 
-		void SetAnalog(bool bInIsAnalog);
-		void SetAttackTime(float InAttackTimeMsec);
-		void SetReleaseTime(float InReleaseTimeMsec);
+		SIGNALPROCESSING_API void SetAnalog(bool bInIsAnalog);
+		SIGNALPROCESSING_API void SetAttackTime(float InAttackTimeMsec);
+		SIGNALPROCESSING_API void SetReleaseTime(float InReleaseTimeMsec);
 
 		/** Get whether set to analog or digital time constant. (True is analog, false is digital) */
 		FORCEINLINE bool GetAnalog() const { return bIsAnalog; }
@@ -69,7 +69,7 @@ namespace Audio
 
 	protected:
 
-		void SetSampleRate(float InSampleRate);
+		SIGNALPROCESSING_API void SetSampleRate(float InSampleRate);
 
 	private:
 
@@ -83,7 +83,7 @@ namespace Audio
 	
 
 	/** Smooths signals using attack and release settings. */
-	class SIGNALPROCESSING_API FAttackReleaseSmoother : public FAttackRelease
+	class FAttackReleaseSmoother : public FAttackRelease
 	{
 	public:
 		/** Construct an FAttackReleaseSmoother.
@@ -94,14 +94,14 @@ namespace Audio
 		 * @param InReleaseTimeMsec - The desired release time in milliseconds.
 		 * @param bInIsAnalog - Whether to model analog RC circuits or use digital models.
 		 */
-		FAttackReleaseSmoother(float InSampleRate, int32 InNumChannels, float InAttackTimeMsec, float InReleaseTimeMsec, bool bInIsAnalog);
+		SIGNALPROCESSING_API FAttackReleaseSmoother(float InSampleRate, int32 InNumChannels, float InAttackTimeMsec, float InReleaseTimeMsec, bool bInIsAnalog);
 
 		/** Process channel interleaved data.
 		 *
 		 * @param InBuffer - Interleaved channel data. 
 		 * @param InNumFrames - Number of frames in InBuffer.
 		 */
-		void ProcessAudio(const float* InBuffer, int32 InNumFrames);
+		SIGNALPROCESSING_API void ProcessAudio(const float* InBuffer, int32 InNumFrames);
 
 		/** Process channel interleaved data.
 		 *
@@ -109,15 +109,15 @@ namespace Audio
 		 * @param InNumFrames - Number of frames in InBuffer.
 		 * @param OutBuffer - Output interleaved envelope data. Should be same size as InBuffer.
 		 */
-		void ProcessAudio(const float* InBuffer, int32 InNumFrames, float* OutBuffer);
+		SIGNALPROCESSING_API void ProcessAudio(const float* InBuffer, int32 InNumFrames, float* OutBuffer);
 
 		/** Retrieve the final values of the envelope for each channel. */
-		const TArray<float>& GetEnvelopeValues() const;
+		SIGNALPROCESSING_API const TArray<float>& GetEnvelopeValues() const;
 
 		/** Set the number of input channels */
-		void SetNumChannels(int32 InNumChannels);
+		SIGNALPROCESSING_API void SetNumChannels(int32 InNumChannels);
 
-		void Reset();
+		SIGNALPROCESSING_API void Reset();
 
 	private:
 		int32 NumChannels;
@@ -125,7 +125,7 @@ namespace Audio
 	};
 
 	/** Compute mean squared using FIR method.  */
-	class SIGNALPROCESSING_API FMeanSquaredFIR
+	class FMeanSquaredFIR
 	{
 		static constexpr int32 DefaultHistoryCapacity = 16384;
 
@@ -137,15 +137,15 @@ namespace Audio
 		 * @param InNumChannels - Number of channels per a frame.
 		 * @param InWindowTimeMsec - Number of milliseconds per a mean squared window.
 		 */
-		FMeanSquaredFIR(float InSampleRate, int32 InNumChannels, float InWindowTimeMsec);
+		SIGNALPROCESSING_API FMeanSquaredFIR(float InSampleRate, int32 InNumChannels, float InWindowTimeMsec);
 
 		/** Set the size of the analysis window. */
-		void SetWindowSize(float InWindowTimeMsec);
+		SIGNALPROCESSING_API void SetWindowSize(float InWindowTimeMsec);
 
 		/** Set the number of input channels. */
-		void SetNumChannels(int32 InNumChannels);
+		SIGNALPROCESSING_API void SetNumChannels(int32 InNumChannels);
 
-		void Reset();
+		SIGNALPROCESSING_API void Reset();
 
 		/** Calculate mean squared per sample.
 		 *
@@ -153,7 +153,7 @@ namespace Audio
 		 * @param InNumFrames - Number of frames in InBuffer.
 		 * @param OutBuffer - Output interleaved data. Should be same size as InBuffer.
 		 */
-		void ProcessAudio(const float* InBuffer, int32 InNumFrames, float* OutBuffer);
+		SIGNALPROCESSING_API void ProcessAudio(const float* InBuffer, int32 InNumFrames, float* OutBuffer);
 
 	private:
 		float SampleRate;
@@ -169,7 +169,7 @@ namespace Audio
 	};
 
 	/** Compute mean squared using IIR method.  */
-	class SIGNALPROCESSING_API FMeanSquaredIIR
+	class FMeanSquaredIIR
 	{
 	public:
 		/** Construct an FMeanSquaredIIR
@@ -178,15 +178,15 @@ namespace Audio
 		 * @param InNumChannels - Number of channels per a frame.
 		 * @param InWindowTimeMsec - Number of milliseconds per a mean squared window.
 		 */
-		FMeanSquaredIIR(float InSampleRate, int32 InNumChannels, float InWindowTimeMsec);
+		SIGNALPROCESSING_API FMeanSquaredIIR(float InSampleRate, int32 InNumChannels, float InWindowTimeMsec);
 
 		/** Set the size of the analysis window. */
-		void SetWindowSize(float InWindowTimeMsec);
+		SIGNALPROCESSING_API void SetWindowSize(float InWindowTimeMsec);
 
 		/** Set the number of input channels. */
-		void SetNumChannels(int32 InNumChannels);
+		SIGNALPROCESSING_API void SetNumChannels(int32 InNumChannels);
 
-		void Reset();
+		SIGNALPROCESSING_API void Reset();
 
 		/** Calculate mean squared per sample.
 		 *
@@ -194,7 +194,7 @@ namespace Audio
 		 * @param InNumFrames - Number of frames in InBuffer.
 		 * @param OutBuffer - Output interleaved data. Should be same size as InBuffer.
 		 */
-		void ProcessAudio(const float* InBuffer, int32 InNumFrames, float* OutBuffer);
+		SIGNALPROCESSING_API void ProcessAudio(const float* InBuffer, int32 InNumFrames, float* OutBuffer);
 
 	private:
 		float SampleRate;
@@ -226,50 +226,50 @@ namespace Audio
 	/** A simple utility that returns a smoothed value given audio input using 
 	 * an RC circuit.  Used for following the envelope of an audio stream.
 	 */
-	class SIGNALPROCESSING_API FEnvelopeFollower
+	class FEnvelopeFollower
 	{
 	public:
 		/** Construct an envelope follower. */
-		FEnvelopeFollower();
+		SIGNALPROCESSING_API FEnvelopeFollower();
 
 		/** Construct an envelope follower. */
-		FEnvelopeFollower(const FEnvelopeFollowerInitParams& InParams);
+		SIGNALPROCESSING_API FEnvelopeFollower(const FEnvelopeFollowerInitParams& InParams);
 
 		/** Initialize the envelope follower. */
-		void Init(const FEnvelopeFollowerInitParams& InParams);
+		SIGNALPROCESSING_API void Init(const FEnvelopeFollowerInitParams& InParams);
 
 		/** Returns the number of channels per an input frame */
-		int32 GetNumChannels() const;
+		SIGNALPROCESSING_API int32 GetNumChannels() const;
 
 		/** Returns the number of frames per a second set on initialization */
-		float GetSampleRate() const;
+		SIGNALPROCESSING_API float GetSampleRate() const;
 
 		/** Returns the envelope follower attack time (how fast the envelope responds to input) */
-		float GetAttackTimeMsec() const;
+		SIGNALPROCESSING_API float GetAttackTimeMsec() const;
 
 		/** Returns the envelope follower release time (how slow the envelope dampens from input) */
-		float GetReleaseTimeMsec() const;
+		SIGNALPROCESSING_API float GetReleaseTimeMsec() const;
 
 		/** Returns whether or not to use analog or digital time constants */
-		bool GetAnalog() const;
+		SIGNALPROCESSING_API bool GetAnalog() const;
 
 		/** Returns the input mode of the envelope follower */
-		EPeakMode::Type GetMode() const;
+		SIGNALPROCESSING_API EPeakMode::Type GetMode() const;
 
 		/** Set the number of channels per an input frame. */
-		void SetNumChannels(int32 InNumChannels);
+		SIGNALPROCESSING_API void SetNumChannels(int32 InNumChannels);
 
 		/** Sets whether or not to use analog or digital time constants */
-		void SetAnalog(bool bInIsAnalog);
+		SIGNALPROCESSING_API void SetAnalog(bool bInIsAnalog);
 
 		/** Sets the envelope follower attack time (how fast the envelope responds to input) */
-		void SetAttackTime(float InAttackTimeMsec);
+		SIGNALPROCESSING_API void SetAttackTime(float InAttackTimeMsec);
 
 		/** Sets the envelope follower release time (how slow the envelope dampens from input) */
-		void SetReleaseTime(float InReleaseTimeMsec);
+		SIGNALPROCESSING_API void SetReleaseTime(float InReleaseTimeMsec);
 
 		/** Sets the input mode of the envelope follower */
-		void SetMode(EPeakMode::Type InMode);
+		SIGNALPROCESSING_API void SetMode(EPeakMode::Type InMode);
 
 		/** Calculate envelope per sample.
 		 *
@@ -277,20 +277,20 @@ namespace Audio
 		 * @param InNumFrames - Number of frames in InBuffer.
 		 * @param OutBuffer - Output interleaved data. Should be same size as InBuffer.
 		 */
-		void ProcessAudio(const float* InBuffer, int32 InNumFrames, float* OutBuffer);
+		SIGNALPROCESSING_API void ProcessAudio(const float* InBuffer, int32 InNumFrames, float* OutBuffer);
 
 		/** Calculate envelope
 		 *
 		 * @param InBuffer - Interleaved input data. 
 		 * @param InNumFrames - Number of frames in InBuffer.
 		 */
-		void ProcessAudio(const float* InBuffer, int32 InNumFrames);
+		SIGNALPROCESSING_API void ProcessAudio(const float* InBuffer, int32 InNumFrames);
 		
 		/** Retrieve the final values of the envelope for each channel. */
-		const TArray<float>& GetEnvelopeValues() const;
+		SIGNALPROCESSING_API const TArray<float>& GetEnvelopeValues() const;
 
 		/** Resets the state of the envelope follower */
-		void Reset();
+		SIGNALPROCESSING_API void Reset();
 
 	private:
 		void ProcessWorkBuffer(const float* InBuffer, int32 InNumFrames);
diff --git a/Engine/Source/Runtime/SignalProcessing/Public/DSP/EventQuantizer.h b/Engine/Source/Runtime/SignalProcessing/Public/DSP/EventQuantizer.h
--- a/Engine/Source/Runtime/SignalProcessing/Public/DSP/EventQuantizer.h
+++ b/Engine/Source/Runtime/SignalProcessing/Public/DSP/EventQuantizer.h
@@ -71,49 +71,49 @@ namespace Audio
 	};
 
 	// Class which handles the details of event quantization.
-	class SIGNALPROCESSING_API FEventQuantizer
+	class FEventQuantizer
 	{
 	public:
-		FEventQuantizer();
-		~FEventQuantizer();
+		SIGNALPROCESSING_API FEventQuantizer();
+		SIGNALPROCESSING_API ~FEventQuantizer();
 
 		// Sets the quantization settings for the event quantizer
-		void SetQuantizationSettings(const FEventQuantizationSettings& QuantizationSettings);
+		SIGNALPROCESSING_API void SetQuantizationSettings(const FEventQuantizationSettings& QuantizationSettings);
 		const FEventQuantizationSettings& GetQuantizationSettings() const { return QuantizationSettings; }
 
 		// Allows continuous control over BPM for the event quantizer
-		void SetBPM(const float InBPM);
+		SIGNALPROCESSING_API void SetBPM(const float InBPM);
 		float GetBPM() const { return QuantizationSettings.BeatsPerMinute; }
 
 		// Set the beat division
-		void SetBeatDivision(const uint16 InBeatDivision);
+		SIGNALPROCESSING_API void SetBeatDivision(const uint16 InBeatDivision);
 		uint16 GetBeatDivision() const { return QuantizationSettings.BeatDivision; }
 
-		float GetPlaybacktimeSeconds() const;
+		SIGNALPROCESSING_API float GetPlaybacktimeSeconds() const;
 
-		uint32 GetDurationInFrames(int32 NumBars, float NumBeats) const;
+		SIGNALPROCESSING_API uint32 GetDurationInFrames(int32 NumBars, float NumBeats) const;
 
 		// Called to perform notifications for any events which happen in the next given number of frames.
 		// This function should be called in an audio buffer render callback.
-		void NotifyEvents(int32 NumFrames);
+		SIGNALPROCESSING_API void NotifyEvents(int32 NumFrames);
 
 		// Enqueues a quantized event
-		void EnqueueEvent(EEventQuantization InQuantization, TFunction<void(uint32 NumFramesOffset)> Lambda);
+		SIGNALPROCESSING_API void EnqueueEvent(EEventQuantization InQuantization, TFunction<void(uint32 NumFramesOffset)> Lambda);
 
 		// Register event listener for specific events
-		void RegisterListenerForEvent(IQuantizedEventListener* InListener, EEventQuantization InQuantization);
+		SIGNALPROCESSING_API void RegisterListenerForEvent(IQuantizedEventListener* InListener, EEventQuantization InQuantization);
 
 		// Unregisters the event listener for all quantization events
-		void UnregisterListenerForEvent(IQuantizedEventListener* InListener);
+		SIGNALPROCESSING_API void UnregisterListenerForEvent(IQuantizedEventListener* InListener);
 
 		// Unregister the event listener for specific quantization event
-		void UnregisterListenerForEvent(IQuantizedEventListener* InListener, EEventQuantization InQuantization);
+		SIGNALPROCESSING_API void UnregisterListenerForEvent(IQuantizedEventListener* InListener, EEventQuantization InQuantization);
 
 	private:
 
-		void SetQuantizationSettingsInternal(const FEventQuantizationSettings& QuantizationSettings);
-		void SetBPMInternal(const float InBPM);
-		void ResetEventState();
+		SIGNALPROCESSING_API void SetQuantizationSettingsInternal(const FEventQuantizationSettings& QuantizationSettings);
+		SIGNALPROCESSING_API void SetBPMInternal(const float InBPM);
+		SIGNALPROCESSING_API void ResetEventState();
 
 		// Struct for defining and hold quantization timing state
 		struct FEventQuantizationState
@@ -136,7 +136,7 @@ namespace Audio
 			{}
 		};
 
-		void NotifyEventForState(FEventQuantizationState& State, EEventQuantization Type, bool bIsQuantizationEvent, int32 NumFrames);
+		SIGNALPROCESSING_API void NotifyEventForState(FEventQuantizationState& State, EEventQuantization Type, bool bIsQuantizationEvent, int32 NumFrames);
 
 		// The frame count of the whole event quantizer since it started. Passed to listeners.
 		uint32 FrameCount;
@@ -166,4 +166,4 @@ namespace Audio
 
 
 
-}
\ No newline at end of file
+}
diff --git a/Engine/Source/Runtime/SignalProcessing/Public/DSP/FFTAlgorithm.h b/Engine/Source/Runtime/SignalProcessing/Public/DSP/FFTAlgorithm.h
--- a/Engine/Source/Runtime/SignalProcessing/Public/DSP/FFTAlgorithm.h
+++ b/Engine/Source/Runtime/SignalProcessing/Public/DSP/FFTAlgorithm.h
@@ -52,11 +52,11 @@ namespace Audio
 	 *
 	 * Interface for FFT algorithm.
 	 */
-	class SIGNALPROCESSING_API IFFTAlgorithm 
+	class IFFTAlgorithm 
 	{
 		public:
 			/** virtual destructor for inheritance. */
-			virtual ~IFFTAlgorithm();
+			SIGNALPROCESSING_API virtual ~IFFTAlgorithm();
 
 			/** Number of elements in FFT */
 			virtual int32 Size() const = 0;
@@ -117,14 +117,14 @@ namespace Audio
 	 *
 	 * Factory interface for creating fft algorithms.
 	 */
-	class SIGNALPROCESSING_API IFFTAlgorithmFactory : public IModularFeature
+	class IFFTAlgorithmFactory : public IModularFeature
 	{
 		public:
 			virtual ~IFFTAlgorithmFactory()
 			{}
 
 			/** Name of modular feature for FFT factory.  */
-			static const FName GetModularFeatureName();
+			static SIGNALPROCESSING_API const FName GetModularFeatureName();
 
 			/** Name of this particular factory. */
 			virtual FName GetFactoryName() const = 0;
@@ -146,11 +146,11 @@ namespace Audio
 	 *
 	 * FFFTFactory creates fft algorithms. It will choose hardware accelerated versions of the FFT when they are available. 
 	 */
-	class SIGNALPROCESSING_API FFFTFactory
+	class FFFTFactory
 	{
 		public:
 			/** This denotes that no specific IFFTAlgorithmFactory is desired. */
-			static const FName AnyAlgorithmFactory;
+			static SIGNALPROCESSING_API const FName AnyAlgorithmFactory;
 
 			/** NewFFTAlgorithm
 			 *
@@ -160,7 +160,7 @@ namespace Audio
 			 * @param InAlgorithmFactoryName - If not equal to FFFTFactory::AnyAlgorithmFactory, will only uses FFT algorithm facotry if IFFTAlgorithmFactory::GetFactoryName() equals InAlgorithmFactoryName.
 			 * @return A TUniquePtr<IFFTAlgorithm> to the created FFT. This pointer can be invalid if an error occured or the fft algorithm could not be created.
 			 */
-			static TUniquePtr<IFFTAlgorithm> NewFFTAlgorithm(const FFFTSettings& InSettings, const FName& InAlgorithmFactoryName = AnyAlgorithmFactory);
+			static SIGNALPROCESSING_API TUniquePtr<IFFTAlgorithm> NewFFTAlgorithm(const FFFTSettings& InSettings, const FName& InAlgorithmFactoryName = AnyAlgorithmFactory);
 
 			/** AreFFTSettingsSupported
 			 *
@@ -170,6 +170,6 @@ namespace Audio
 			 * @param InAlgorithmFactoryName - If not equal to FFFTFactory::AnyAlgorithmFactory, will only uses FFT algorithm facotry if IFFTAlgorithmFactory::GetFactoryName() equals InAlgorithmFactoryName.
 			 * @return True if settings are supported. Otherwise false.
 			 */
-			static bool AreFFTSettingsSupported(const FFFTSettings& InSettings, const FName& InAlgorithmFactoryName = AnyAlgorithmFactory);
+			static SIGNALPROCESSING_API bool AreFFTSettingsSupported(const FFFTSettings& InSettings, const FName& InAlgorithmFactoryName = AnyAlgorithmFactory);
 	};
 }
diff --git a/Engine/Source/Runtime/SignalProcessing/Public/DSP/FeedbackDelayNetwork.h b/Engine/Source/Runtime/SignalProcessing/Public/DSP/FeedbackDelayNetwork.h
--- a/Engine/Source/Runtime/SignalProcessing/Public/DSP/FeedbackDelayNetwork.h
+++ b/Engine/Source/Runtime/SignalProcessing/Public/DSP/FeedbackDelayNetwork.h
@@ -11,7 +11,7 @@
 namespace Audio
 {
 	// Delay settings of 4 channel feedback delay network
-	struct SIGNALPROCESSING_API FFDNDelaySettings
+	struct FFDNDelaySettings
 	{
 		// Number of samples for delay line 0
 		int32 APF0DelayNumSamples;
@@ -23,13 +23,13 @@ namespace Audio
 		int32 APF3DelayNumSamples;
 		
 		// Default delay line settings for left channel.
-		static FFDNDelaySettings DefaultLeftDelays(float InSampleRate);
+		static SIGNALPROCESSING_API FFDNDelaySettings DefaultLeftDelays(float InSampleRate);
 		// Default delay line settings for right channel. 
-		static FFDNDelaySettings DefaultRightDelays(float InSampleRate);
+		static SIGNALPROCESSING_API FFDNDelaySettings DefaultRightDelays(float InSampleRate);
 	};
 
 	// Filter coefficients of 4 channel feedback delay network.
-	struct SIGNALPROCESSING_API FFDNCoefficients
+	struct FFDNCoefficients
 	{
 		// Sample multiplier of input samples before they enter the delay line.
 		float InputScale;
@@ -44,26 +44,26 @@ namespace Audio
 	};
 
 	// 4 channel feedback delay network (FDN) for artificial reverberation.
-	class SIGNALPROCESSING_API FFeedbackDelayNetwork
+	class FFeedbackDelayNetwork
 	{
 	public:
 		// InMaxNumInternalBufferSamples controls the internal buffer size used for vector operations. 
 		// InSettings controls the delay line lengths.
-		FFeedbackDelayNetwork(int32 InMaxNumInternalBufferSamples, const FFDNDelaySettings& InSettings);
+		SIGNALPROCESSING_API FFeedbackDelayNetwork(int32 InMaxNumInternalBufferSamples, const FFDNDelaySettings& InSettings);
 
-		~FFeedbackDelayNetwork();
+		SIGNALPROCESSING_API ~FFeedbackDelayNetwork();
 
 		// Sets the coefficient values of the all pass filters, low pass filters, input scalers and feedback scalers.
-		void SetCoefficients(const FFDNCoefficients& InCoefficients);
+		SIGNALPROCESSING_API void SetCoefficients(const FFDNCoefficients& InCoefficients);
 
 		// Generates artificial reverberation for InSamples and places results in OutSamples.
-		void ProcessAudio(const FAlignedFloatBuffer& InSamples, FAlignedFloatBuffer& OutSamples);
+		SIGNALPROCESSING_API void ProcessAudio(const FAlignedFloatBuffer& InSamples, FAlignedFloatBuffer& OutSamples);
 
 		// Sets internal audio samples to silence. 
-		void FlushAudio();
+		SIGNALPROCESSING_API void FlushAudio();
 
 		// Returns the actual number of samples in an internal buffer.
-		int32 GetNumInternalBufferSamples() const;
+		SIGNALPROCESSING_API int32 GetNumInternalBufferSamples() const;
 
 
 	private:
diff --git a/Engine/Source/Runtime/SignalProcessing/Public/DSP/FileEncoder.h b/Engine/Source/Runtime/SignalProcessing/Public/DSP/FileEncoder.h
--- a/Engine/Source/Runtime/SignalProcessing/Public/DSP/FileEncoder.h
+++ b/Engine/Source/Runtime/SignalProcessing/Public/DSP/FileEncoder.h
@@ -8,26 +8,26 @@
 
 namespace Audio
 {
-	class SIGNALPROCESSING_API FAudioFileWriter
+	class FAudioFileWriter
 	{
 	public:
 		// Constructor. Takes an absolute file path and metadata about the audio in question.
 		// Immediately allocates all data necessary.
-		FAudioFileWriter(const FString& InPath, const FSoundQualityInfo& InInfo);
+		SIGNALPROCESSING_API FAudioFileWriter(const FString& InPath, const FSoundQualityInfo& InInfo);
 
 		// Calling the destructor on this class finalizes and closes out the file.
-		~FAudioFileWriter();
+		SIGNALPROCESSING_API ~FAudioFileWriter();
 
 		// Returns file information.
-		void GetFileInfo(FSoundQualityInfo& OutInfo);
+		SIGNALPROCESSING_API void GetFileInfo(FSoundQualityInfo& OutInfo);
 
 		// Use this to push audio to the file writer.
 		// If you'd like to move encoding and file writing to a separate thread from PushAudio,
 		// Call this with bEncodeIfPossible = false.
-		bool PushAudio(const float* InAudio, int32 NumSamples, bool bEncodeIfPossible = true);
+		SIGNALPROCESSING_API bool PushAudio(const float* InAudio, int32 NumSamples, bool bEncodeIfPossible = true);
 
 		// If PushAudio is called with bEncodeIfPossible set to false, this will need to be called.
-		bool EncodeIfPossible();
+		SIGNALPROCESSING_API bool EncodeIfPossible();
 
 	private:
 		FAudioFileWriter();
diff --git a/Engine/Source/Runtime/SignalProcessing/Public/DSP/Filter.h b/Engine/Source/Runtime/SignalProcessing/Public/DSP/Filter.h
--- a/Engine/Source/Runtime/SignalProcessing/Public/DSP/Filter.h
+++ b/Engine/Source/Runtime/SignalProcessing/Public/DSP/Filter.h
@@ -29,54 +29,54 @@ namespace Audio
 
 	// Biquad filter class which wraps a biquad filter struct
 	// Handles multi-channel audio to avoid calculating filter coefficients for multiple channels of audio.
-	class SIGNALPROCESSING_API FBiquadFilter
+	class FBiquadFilter
 	{
 	public:
 		// Constructor
-		FBiquadFilter();
+		SIGNALPROCESSING_API FBiquadFilter();
 		// Destructor
-		virtual ~FBiquadFilter();
+		SIGNALPROCESSING_API virtual ~FBiquadFilter();
 
 		// Initialize the filter
-		void Init(const float InSampleRate, const int32 InNumChannels, const EBiquadFilter::Type InType, const float InCutoffFrequency = 20000.0f, const float InBandwidth = 2.0f, const float InGain = 0.0f);
+		SIGNALPROCESSING_API void Init(const float InSampleRate, const int32 InNumChannels, const EBiquadFilter::Type InType, const float InCutoffFrequency = 20000.0f, const float InBandwidth = 2.0f, const float InGain = 0.0f);
 
 		// Returns number of channels initialized with
-		int32 GetNumChannels() const;
+		SIGNALPROCESSING_API int32 GetNumChannels() const;
 
 		// Resets the filter state
-		void Reset();
+		SIGNALPROCESSING_API void Reset();
 
 		// Processes a single frame of audio
-		void ProcessAudioFrame(const float* InFrame, float* OutFrame);
+		SIGNALPROCESSING_API void ProcessAudioFrame(const float* InFrame, float* OutFrame);
 
 		// Process a buffer of audio
-		void ProcessAudio(const float* InBuffer, const int32 InNumSamples, float* OutBuffer);
+		SIGNALPROCESSING_API void ProcessAudio(const float* InBuffer, const int32 InNumSamples, float* OutBuffer);
 
 		// Sets all filter parameters with one function
-		void SetParams(const EBiquadFilter::Type InFilterType, const float InCutoffFrequency, const float InBandwidth, const float InGainDB);
+		SIGNALPROCESSING_API void SetParams(const EBiquadFilter::Type InFilterType, const float InCutoffFrequency, const float InBandwidth, const float InGainDB);
 
 		// Sets the type of the filter to use
-		void SetType(const EBiquadFilter::Type InType);
+		SIGNALPROCESSING_API void SetType(const EBiquadFilter::Type InType);
 
 		// Sets the filter frequency
-		void SetFrequency(const float InCutoffFrequency);
+		SIGNALPROCESSING_API void SetFrequency(const float InCutoffFrequency);
 
 		// Sets the bandwidth (octaves) of the filter
-		void SetBandwidth(const float InBandwidth);
+		SIGNALPROCESSING_API void SetBandwidth(const float InBandwidth);
 
 		// Sets the gain of the filter in decibels
-		void SetGainDB(const float InGainDB);
+		SIGNALPROCESSING_API void SetGainDB(const float InGainDB);
 
 		// Sets whether or no this filter is enabled (if disabled audio is passed through)
-		void SetEnabled(const bool bInEnabled);
+		SIGNALPROCESSING_API void SetEnabled(const bool bInEnabled);
 
 	protected:
 
 		// Function computes biquad coefficients based on current filter settings
-		void CalculateBiquadCoefficients();
+		SIGNALPROCESSING_API void CalculateBiquadCoefficients();
 
 		// Function used to clamp the cutoff frequency.
-		float ClampCutoffFrequency(float InCutoffFrequency);
+		SIGNALPROCESSING_API float ClampCutoffFrequency(float InCutoffFrequency);
 
 		// What kind of filter to use when computing coefficients
 		EBiquadFilter::Type FilterType;
@@ -119,26 +119,26 @@ namespace Audio
 	static const int32 MaxFilterChannels = 8;
 
 	// Base class for filters usable in synthesis
-	class SIGNALPROCESSING_API IFilter
+	class IFilter
 	{
 	public:
-		IFilter();
-		virtual ~IFilter();
+		SIGNALPROCESSING_API IFilter();
+		SIGNALPROCESSING_API virtual ~IFilter();
 
 		// Initialize the filter
-		virtual void Init(const float InSampleRate, const int32 InNumChannels, const int32 InVoiceId, FModulationMatrix* InModMatrix = nullptr);
+		SIGNALPROCESSING_API virtual void Init(const float InSampleRate, const int32 InNumChannels, const int32 InVoiceId, FModulationMatrix* InModMatrix = nullptr);
 
 		// Sets the cutoff frequency of the filter.
-		virtual void SetFrequency(const float InCutoffFrequency);
+		SIGNALPROCESSING_API virtual void SetFrequency(const float InCutoffFrequency);
 
 		// Sets an external modulated frequency
-		virtual void SetFrequencyMod(const float InModFrequency);
+		SIGNALPROCESSING_API virtual void SetFrequencyMod(const float InModFrequency);
 
 		// Sets the quality/resonance of the filter
-		virtual void SetQ(const float InQ);
+		SIGNALPROCESSING_API virtual void SetQ(const float InQ);
 
 		// Sets an external modulated quality/resonance of the filter
-		virtual void SetQMod(const float InModQ);
+		SIGNALPROCESSING_API virtual void SetQMod(const float InModQ);
 
 		// Sets the filter saturation (not used on all filters)
 		virtual void SetSaturation(const float InSaturation) {}
@@ -150,13 +150,13 @@ namespace Audio
 		virtual void SetPassBandGainCompensation(const float InPassBandGainCompensation) {}
 
 		// Sets the filter type
-		virtual void SetFilterType(const EFilter::Type InFilterType);
+		SIGNALPROCESSING_API virtual void SetFilterType(const EFilter::Type InFilterType);
 
 		// Reset the filter
-		virtual void Reset();
+		SIGNALPROCESSING_API virtual void Reset();
 
 		// Updates the filter
-		virtual void Update();
+		SIGNALPROCESSING_API virtual void Update();
 
 		// Processes a single frame of audio. Number of channels MUST be what was used during filter initialization.
 		virtual void ProcessAudioFrame(const float* InFrame, float* OutFrame)
@@ -228,17 +228,17 @@ namespace Audio
 
 	// A virtual analog one-pole filter.
 	// Defaults to a low-pass mode, but can be switched to high-pass
-	class SIGNALPROCESSING_API FOnePoleFilter : public IFilter
+	class FOnePoleFilter : public IFilter
 	{
 	public:
-		FOnePoleFilter();
-		virtual ~FOnePoleFilter();
+		SIGNALPROCESSING_API FOnePoleFilter();
+		SIGNALPROCESSING_API virtual ~FOnePoleFilter();
 
-		virtual void Init(const float InSampleRate, const int32 InNumChannels, const int32 InVoiceId = 0, FModulationMatrix* InModMatrix = nullptr) override;
-		virtual void Reset() override;
-		virtual void Update() override;
-		virtual void ProcessAudioFrame(const float* InFrame, float* OutFrame) override;
-		virtual void ProcessAudio(const float* InSamples, const int32 InNumSamples, float* OutSamples) override;
+		SIGNALPROCESSING_API virtual void Init(const float InSampleRate, const int32 InNumChannels, const int32 InVoiceId = 0, FModulationMatrix* InModMatrix = nullptr) override;
+		SIGNALPROCESSING_API virtual void Reset() override;
+		SIGNALPROCESSING_API virtual void Update() override;
+		SIGNALPROCESSING_API virtual void ProcessAudioFrame(const float* InFrame, float* OutFrame) override;
+		SIGNALPROCESSING_API virtual void ProcessAudio(const float* InSamples, const int32 InNumSamples, float* OutSamples) override;
 
 		void SetCoefficient(const float InCoefficient) { A0 = InCoefficient; }
 		float GetCoefficient() const { return A0; }
@@ -250,19 +250,19 @@ namespace Audio
 		float* Z1;
 	};
 
-	class SIGNALPROCESSING_API FStateVariableFilter : public IFilter
+	class FStateVariableFilter : public IFilter
 	{
 	public:
-		FStateVariableFilter();
-		virtual ~FStateVariableFilter();
+		SIGNALPROCESSING_API FStateVariableFilter();
+		SIGNALPROCESSING_API virtual ~FStateVariableFilter();
 
-		virtual void Init(const float InSampleRate, const int32 InNumChannels, const int32 InVoiceId = 0, FModulationMatrix* InModMatrix = nullptr) override;
-		virtual void SetBandStopControl(const float InBandStop) override;
-		virtual void Reset() override;
-		virtual void Update() override;
-		virtual void ProcessAudio(const float* InSamples, const int32 InNumSamples, float* OutSamples) override;
+		SIGNALPROCESSING_API virtual void Init(const float InSampleRate, const int32 InNumChannels, const int32 InVoiceId = 0, FModulationMatrix* InModMatrix = nullptr) override;
+		SIGNALPROCESSING_API virtual void SetBandStopControl(const float InBandStop) override;
+		SIGNALPROCESSING_API virtual void Reset() override;
+		SIGNALPROCESSING_API virtual void Update() override;
+		SIGNALPROCESSING_API virtual void ProcessAudio(const float* InSamples, const int32 InNumSamples, float* OutSamples) override;
 
-		void ProcessAudio(const float* InSamples, const int32 InNumSamples
+		SIGNALPROCESSING_API void ProcessAudio(const float* InSamples, const int32 InNumSamples
 			, float* LpfOutput, float* HpfOutput, float* BpfOutput, float* BsfOutput);
 
 	protected:
@@ -285,19 +285,19 @@ namespace Audio
 		TArray<FFilterState> FilterState;
 	};
 
-	class SIGNALPROCESSING_API FLadderFilter : public IFilter
+	class FLadderFilter : public IFilter
 	{
 	public:
-		FLadderFilter();
-		virtual ~FLadderFilter();
+		SIGNALPROCESSING_API FLadderFilter();
+		SIGNALPROCESSING_API virtual ~FLadderFilter();
 
-		virtual void Init(const float InSampleRate, const int32 InNumChannels, const int32 InVoiceId = 0, FModulationMatrix* InModMatrix = nullptr) override;
-		virtual void Reset() override;
-		virtual void Update() override;
-		virtual void SetQ(const float InQ) override;
-		virtual void SetPassBandGainCompensation(const float InPassBandGainCompensation) override;
+		SIGNALPROCESSING_API virtual void Init(const float InSampleRate, const int32 InNumChannels, const int32 InVoiceId = 0, FModulationMatrix* InModMatrix = nullptr) override;
+		SIGNALPROCESSING_API virtual void Reset() override;
+		SIGNALPROCESSING_API virtual void Update() override;
+		SIGNALPROCESSING_API virtual void SetQ(const float InQ) override;
+		SIGNALPROCESSING_API virtual void SetPassBandGainCompensation(const float InPassBandGainCompensation) override;
 
-		virtual void ProcessAudio(const float* InSamples, const int32 InNumSamples, float* OutSamples) override;
+		SIGNALPROCESSING_API virtual void ProcessAudio(const float* InSamples, const int32 InNumSamples, float* OutSamples) override;
 
 
 	protected:
diff --git a/Engine/Source/Runtime/SignalProcessing/Public/DSP/Flanger.h b/Engine/Source/Runtime/SignalProcessing/Public/DSP/Flanger.h
--- a/Engine/Source/Runtime/SignalProcessing/Public/DSP/Flanger.h
+++ b/Engine/Source/Runtime/SignalProcessing/Public/DSP/Flanger.h
@@ -8,24 +8,24 @@
 
 namespace Audio
 {
-	class SIGNALPROCESSING_API FFlanger
+	class FFlanger
 	{
 	public: 
-		FFlanger();
-		~FFlanger();
+		SIGNALPROCESSING_API FFlanger();
+		SIGNALPROCESSING_API ~FFlanger();
 
-		void Init(const float InSampleRate);
+		SIGNALPROCESSING_API void Init(const float InSampleRate);
 
-		void SetModulationRate(const float InModulationRate);
-		void SetModulationDepth(const float InModulationDepth);
-		void SetCenterDelay(const float InCenterDelay);
-		void SetMixLevel(const float InMixLevel);
+		SIGNALPROCESSING_API void SetModulationRate(const float InModulationRate);
+		SIGNALPROCESSING_API void SetModulationDepth(const float InModulationDepth);
+		SIGNALPROCESSING_API void SetCenterDelay(const float InCenterDelay);
+		SIGNALPROCESSING_API void SetMixLevel(const float InMixLevel);
 
-		void ProcessAudio(const FAlignedFloatBuffer& InBuffer, const int32 InNumSamples, FAlignedFloatBuffer& OutBuffer);
+		SIGNALPROCESSING_API void ProcessAudio(const FAlignedFloatBuffer& InBuffer, const int32 InNumSamples, FAlignedFloatBuffer& OutBuffer);
 	protected: 
-		static const float MaxDelaySec;
-		static const float MaxModulationRate;
-		static const float MaxCenterDelay;
+		static SIGNALPROCESSING_API const float MaxDelaySec;
+		static SIGNALPROCESSING_API const float MaxModulationRate;
+		static SIGNALPROCESSING_API const float MaxCenterDelay;
 
 		// LFO parameters
 		// LFO frequency 
diff --git a/Engine/Source/Runtime/SignalProcessing/Public/DSP/FloatArrayMath.h b/Engine/Source/Runtime/SignalProcessing/Public/DSP/FloatArrayMath.h
--- a/Engine/Source/Runtime/SignalProcessing/Public/DSP/FloatArrayMath.h
+++ b/Engine/Source/Runtime/SignalProcessing/Public/DSP/FloatArrayMath.h
@@ -254,7 +254,7 @@ namespace Audio
 	 *  It works with non-sparse and non-contiguous kernels as well, but will be more computationally 
 	 *  expensive than a naive implementation. Also, only takes advantage of sparse contiguous rows, not columns.
 	 */
-	class SIGNALPROCESSING_API FContiguousSparse2DKernelTransform
+	class FContiguousSparse2DKernelTransform
 	{
 		struct FRow
 		{
@@ -271,14 +271,14 @@ namespace Audio
 		 * NumInElements sets the expected number of input array elements as well as the number of elements in a row.
 		 * NumOutElements sets the number of output array elements as well as the number or rows.
 		 */
-		FContiguousSparse2DKernelTransform(const int32 NumInElements, const int32 NumOutElements);
-		virtual ~FContiguousSparse2DKernelTransform();
+		SIGNALPROCESSING_API FContiguousSparse2DKernelTransform(const int32 NumInElements, const int32 NumOutElements);
+		SIGNALPROCESSING_API virtual ~FContiguousSparse2DKernelTransform();
 
 		/** Returns the required size of the input array */
-		int32 GetNumInElements() const;
+		SIGNALPROCESSING_API int32 GetNumInElements() const;
 
 		/** Returns the size of the output array */
-		int32 GetNumOutElements() const;
+		SIGNALPROCESSING_API int32 GetNumOutElements() const;
 	
 		/** Set the kernel values for an individual row.
 		 *
@@ -286,28 +286,28 @@ namespace Audio
 		 *  StartIndex denotes the offset into the row where the OffsetValues will be inserted.
 		 *  OffsetValues contains the contiguous chunk of values which represent all the nonzero elements in the row.
 		 */
-		void SetRow(const int32 RowIndex, const int32 StartIndex, TArrayView<const float> OffsetValues);
+		SIGNALPROCESSING_API void SetRow(const int32 RowIndex, const int32 StartIndex, TArrayView<const float> OffsetValues);
 
 		/** Transforms the input array given the kernel.
 		 *
 		 *  InView is the array to be transformed. It must have `NumInElements` number of elements.
 		 *  OutArray is the transformed array. It will have `NumOutElements` number of elements.
 		 */
-		void TransformArray(TArrayView<const float> InView, TArray<float>& OutArray) const;
+		SIGNALPROCESSING_API void TransformArray(TArrayView<const float> InView, TArray<float>& OutArray) const;
 
 		/** Transforms the input array given the kernel.
 		 *
 		 *  InView is the array to be transformed. It must have `NumInElements` number of elements.
 		 *  OutArray is the transformed array. It will have `NumOutElements` number of elements.
 		 */
-		void TransformArray(TArrayView<const float> InView, FAlignedFloatBuffer& OutArray) const;
+		SIGNALPROCESSING_API void TransformArray(TArrayView<const float> InView, FAlignedFloatBuffer& OutArray) const;
 
 		/** Transforms the input array given the kernel.
 		 *
 		 *  InArray is the array to be transformed. It must have `NumInElements` number of elements.
 		 *  OutArray is the transformed array. It must be allocated to hold at least NumOutElements. 
 		 */
-		void TransformArray(const float* InArray, float* OutArray) const;
+		SIGNALPROCESSING_API void TransformArray(const float* InArray, float* OutArray) const;
 
 	private:
 
diff --git a/Engine/Source/Runtime/SignalProcessing/Public/DSP/FoldbackDistortion.h b/Engine/Source/Runtime/SignalProcessing/Public/DSP/FoldbackDistortion.h
--- a/Engine/Source/Runtime/SignalProcessing/Public/DSP/FoldbackDistortion.h
+++ b/Engine/Source/Runtime/SignalProcessing/Public/DSP/FoldbackDistortion.h
@@ -9,35 +9,35 @@ namespace Audio
 {
 	// Foldback distortion effect
 	// https://en.wikipedia.org/wiki/Foldback_(power_supply_design)
-	class SIGNALPROCESSING_API FFoldbackDistortion
+	class FFoldbackDistortion
 	{
 	public:
 		// Constructor
-		FFoldbackDistortion();
+		SIGNALPROCESSING_API FFoldbackDistortion();
 
 		// Destructor
-		~FFoldbackDistortion();
+		SIGNALPROCESSING_API ~FFoldbackDistortion();
 
 		// Initialize the equalizer
-		void Init(const float InSampleRate, const int32 InNumChannels);
+		SIGNALPROCESSING_API void Init(const float InSampleRate, const int32 InNumChannels);
 
 		// Sets the foldback distortion threshold
-		void SetThresholdDb(const float InThresholdDb);
+		SIGNALPROCESSING_API void SetThresholdDb(const float InThresholdDb);
 
 		// Sets the input gain
-		void SetInputGainDb(const float InInputGainDb);
+		SIGNALPROCESSING_API void SetInputGainDb(const float InInputGainDb);
 
 		// Sets the output gain
-		void SetOutputGainDb(const float InOutputGainDb);
+		SIGNALPROCESSING_API void SetOutputGainDb(const float InOutputGainDb);
 
 		// Processes a single audio sample
-		float ProcessAudioSample(const float InSample);
+		SIGNALPROCESSING_API float ProcessAudioSample(const float InSample);
 
 		// Processes a mono stream
-		void ProcessAudioFrame(const float* InFrame, float* OutFrame);
+		SIGNALPROCESSING_API void ProcessAudioFrame(const float* InFrame, float* OutFrame);
 
 		// Processes a stereo stream
-		void ProcessAudio(const float* InBuffer, const int32 InNumSamples, float* OutBuffer);
+		SIGNALPROCESSING_API void ProcessAudio(const float* InBuffer, const int32 InNumSamples, float* OutBuffer);
 
 	private:
 		// Threshold to check before folding audio back on itself
diff --git a/Engine/Source/Runtime/SignalProcessing/Public/DSP/GrainDelay.h b/Engine/Source/Runtime/SignalProcessing/Public/DSP/GrainDelay.h
--- a/Engine/Source/Runtime/SignalProcessing/Public/DSP/GrainDelay.h
+++ b/Engine/Source/Runtime/SignalProcessing/Public/DSP/GrainDelay.h
@@ -10,38 +10,38 @@ namespace Audio
 {
 	namespace GrainDelay
 	{
-		class SIGNALPROCESSING_API FGrainDelay
+		class FGrainDelay
 		{
 		public:
-			FGrainDelay(const float InSampleRate);
-			~FGrainDelay();
+			SIGNALPROCESSING_API FGrainDelay(const float InSampleRate);
+			SIGNALPROCESSING_API ~FGrainDelay();
 
-			void Reset();
+			SIGNALPROCESSING_API void Reset();
 
 			// Helpers to clamp audio
-			float GetGrainDelayClamped(const float InDelay) const;
-			float GetGrainDurationClamped(const float InDuration) const;
-			float GetGrainDelayRatioClamped(const float InGrainDelayRatio) const;
-			float GetGrainPitchShiftClamped(const float InPitchShift) const;
-			float GetGrainPitchShiftFrameRatio(const float InPitchShift) const;
+			SIGNALPROCESSING_API float GetGrainDelayClamped(const float InDelay) const;
+			SIGNALPROCESSING_API float GetGrainDurationClamped(const float InDuration) const;
+			SIGNALPROCESSING_API float GetGrainDelayRatioClamped(const float InGrainDelayRatio) const;
+			SIGNALPROCESSING_API float GetGrainPitchShiftClamped(const float InPitchShift) const;
+			SIGNALPROCESSING_API float GetGrainPitchShiftFrameRatio(const float InPitchShift) const;
 
 			// Dynamically sets the max grains 
-			void SetMaxGrains(const int32 InMaxGrains); 
+			SIGNALPROCESSING_API void SetMaxGrains(const int32 InMaxGrains); 
 			
 			// Sets the grain envelope 
-			void SetGrainEnvelope(const Audio::Grain::EEnvelope InGrainEnvelope);
+			SIGNALPROCESSING_API void SetGrainEnvelope(const Audio::Grain::EEnvelope InGrainEnvelope);
 
 			// Sets the feedback amount (how much output grain audio feeds into the delay line)
-			void SetFeedbackAmount(float InFeedbackAmount);
+			SIGNALPROCESSING_API void SetFeedbackAmount(float InFeedbackAmount);
 
 			// Sets the base pitch shift ratio of all grains 
-			void SetGrainBasePitchShiftRatio(const float InPitchRatioBase);
+			SIGNALPROCESSING_API void SetGrainBasePitchShiftRatio(const float InPitchRatioBase);
 
 			// Spawns a new grain with the given parameters
-			void SpawnGrain(const float InDelay, const float InDuration, const float InPitchShiftRatioOffset);
+			SIGNALPROCESSING_API void SpawnGrain(const float InDelay, const float InDuration, const float InPitchShiftRatioOffset);
 
 			// Synthesize audio in the given frame range from input audio. Writes to OutAudioBuffer
-			void SynthesizeAudio(const int32 StartFrame, const int32 EndFrame, const float* InAudioBuffer, float* OutAudioBuffer);
+			SIGNALPROCESSING_API void SynthesizeAudio(const int32 StartFrame, const int32 EndFrame, const float* InAudioBuffer, float* OutAudioBuffer);
 			
 		private:
 			// Synthesize a single frame of audio
diff --git a/Engine/Source/Runtime/SignalProcessing/Public/DSP/Granulator.h b/Engine/Source/Runtime/SignalProcessing/Public/DSP/Granulator.h
--- a/Engine/Source/Runtime/SignalProcessing/Public/DSP/Granulator.h
+++ b/Engine/Source/Runtime/SignalProcessing/Public/DSP/Granulator.h
@@ -42,14 +42,14 @@ namespace Audio
 	};
 
 	// Simple class that generates an envelope and lets you retrieve interpolated values at any given fraction
-	class SIGNALPROCESSING_API FGrainEnvelope
+	class FGrainEnvelope
 	{
 	public:
-		FGrainEnvelope();
-		~FGrainEnvelope();
+		SIGNALPROCESSING_API FGrainEnvelope();
+		SIGNALPROCESSING_API ~FGrainEnvelope();
 
-		void GenerateEnvelope(const EGrainEnvelopeType EnvelopeType, const int32 NumFrames);
-		float GetValue(const float Fraction) const;
+		SIGNALPROCESSING_API void GenerateEnvelope(const EGrainEnvelopeType EnvelopeType, const int32 NumFrames);
+		SIGNALPROCESSING_API float GetValue(const float Fraction) const;
 
 	private:
 		EGrainEnvelopeType CurrentType;
@@ -83,44 +83,44 @@ namespace Audio
 	class FGranularSynth;
 
 	// Class representing a grain of audio
-	class SIGNALPROCESSING_API FGrain
+	class FGrain
 	{
 	public:
-		FGrain(const int32 InGrainId, FGranularSynth* InParent);
-		~FGrain();
+		SIGNALPROCESSING_API FGrain(const int32 InGrainId, FGranularSynth* InParent);
+		SIGNALPROCESSING_API ~FGrain();
 
 		// Plays the grain with the supplied grain data
-		void Play(const FGrainData& InGrainData);
+		SIGNALPROCESSING_API void Play(const FGrainData& InGrainData);
 
 		// Changes the oscillator type for the grain (if the grain is in synth mode)
-		void SetOscType(const EOsc::Type InType);
+		SIGNALPROCESSING_API void SetOscType(const EOsc::Type InType);
 
 		// Sets the oscillator frequency
-		void SetOscFrequency(const float InFrequency);
+		SIGNALPROCESSING_API void SetOscFrequency(const float InFrequency);
 
 		// Sets the oscillator frequency modulation
-		void SetOscFrequencyModuation(const float InFrequencyModulation);
+		SIGNALPROCESSING_API void SetOscFrequencyModuation(const float InFrequencyModulation);
 
 		// Sets the grain pitch modulation 
-		void SetPitchModulation(const float InPitchModulation);
+		SIGNALPROCESSING_API void SetPitchModulation(const float InPitchModulation);
 
 		// Sets the grain modulation
-		void SetVolumeModulation(const float InVolumeModulation);
+		SIGNALPROCESSING_API void SetVolumeModulation(const float InVolumeModulation);
 
 		// Sets the pan modulation angle
-		void SetPanModulation(const float InPanModulation);
+		SIGNALPROCESSING_API void SetPanModulation(const float InPanModulation);
 
 		// Changes how quickly the grain reads through envelope
-		void SetDurationScale(const float InDurationScale);
+		SIGNALPROCESSING_API void SetDurationScale(const float InDurationScale);
 
 		// Queries if this grain is finished playing and needs to be reaped
-		bool IsDone() const;
+		SIGNALPROCESSING_API bool IsDone() const;
 
 		// Generates the next frame from the grain
-		bool GenerateFrame(float* OutStereoFrame);
+		SIGNALPROCESSING_API bool GenerateFrame(float* OutStereoFrame);
 
 	protected:
-		float GetEnvelopeValue();
+		SIGNALPROCESSING_API float GetEnvelopeValue();
 
 		// Grain id
 		int32 GrainId;
@@ -166,104 +166,104 @@ namespace Audio
 	};
 
 	// A stereo granulator 
-	class SIGNALPROCESSING_API FGranularSynth
+	class FGranularSynth
 	{
 	public:
-		FGranularSynth();
-		~FGranularSynth();
+		SIGNALPROCESSING_API FGranularSynth();
+		SIGNALPROCESSING_API ~FGranularSynth();
 
-		void Init(const int32 InSampleRate, const int32 InNumInitialGrains);
+		SIGNALPROCESSING_API void Init(const int32 InSampleRate, const int32 InNumInitialGrains);
 		
 		// Loads a sound wave to use for granular synth mode
-		void LoadSampleBuffer(const TSampleBuffer<int16>& InSampleBuffer);
+		SIGNALPROCESSING_API void LoadSampleBuffer(const TSampleBuffer<int16>& InSampleBuffer);
 
 		// Plays a granular synthesis "Note"
-		void NoteOn(const uint32 InMidiNote, const float InVelocity, const float InDurationSec = INDEX_NONE);
+		SIGNALPROCESSING_API void NoteOn(const uint32 InMidiNote, const float InVelocity, const float InDurationSec = INDEX_NONE);
 
 		// Note off, triggers release envelope
-		void NoteOff(const uint32 InMidiNote, const bool bKill);
+		SIGNALPROCESSING_API void NoteOff(const uint32 InMidiNote, const bool bKill);
 
 		// Sets the granular synth attack time
-		void SetAttackTime(const float InAttackTimeMSec);
+		SIGNALPROCESSING_API void SetAttackTime(const float InAttackTimeMSec);
 
 		// Sets the granular synth decay time
-		void SetDecayTime(const float InDecayTimeMSec);
+		SIGNALPROCESSING_API void SetDecayTime(const float InDecayTimeMSec);
 
 		// Sets the granular synth sustain gain
-		void SetSustainGain(const float InSustainGain);
+		SIGNALPROCESSING_API void SetSustainGain(const float InSustainGain);
 
 		// Sets the granular synth releas etime
-		void SetReleaseTime(const float InReleaseTimeMSec);
+		SIGNALPROCESSING_API void SetReleaseTime(const float InReleaseTimeMSec);
 
 		// Seeks the loaded buffer used for granulation. Grains will spawn from this location
-		void SeekTime(const float InTimeSec, const float LerpTimeSec = 0.0f, const ESeekType::Type InSeekType = ESeekType::FromBeginning);
+		SIGNALPROCESSING_API void SeekTime(const float InTimeSec, const float LerpTimeSec = 0.0f, const ESeekType::Type InSeekType = ESeekType::FromBeginning);
 
 		// Sets whether or not the buffer playback advances on its own or if it just sits in one place
-		void SetScrubMode(const bool bIsScrubMode);
+		SIGNALPROCESSING_API void SetScrubMode(const bool bIsScrubMode);
 
 		// Sets how fast the granular play head for granulation is is played (and in what direction)
-		void SetPlaybackSpeed(const float InPlaybackSpeed);
+		SIGNALPROCESSING_API void SetPlaybackSpeed(const float InPlaybackSpeed);
 
 		// The rate at which new grains are attempted to be spawned
-		void SetGrainsPerSecond(const float InNumberOfGrainsPerSecond);
+		SIGNALPROCESSING_API void SetGrainsPerSecond(const float InNumberOfGrainsPerSecond);
 
 		// The probability at which a grain will occur when a grain tries to spawn. Allows for sporatic grain generation.
-		void SetGrainProbability(const float InGrainProbability);
+		SIGNALPROCESSING_API void SetGrainProbability(const float InGrainProbability);
 
 		// Sets the envelope type to use for new grains. Will instantly switch all grains to this envelope type so may cause discontinuities if switched while playing.
-		void SetGrainEnvelopeType(const EGrainEnvelopeType InGrainEnvelopeType);
+		SIGNALPROCESSING_API void SetGrainEnvelopeType(const EGrainEnvelopeType InGrainEnvelopeType);
 
 		// Sets the grain oscillator type (for use with granular synthesis mode)
-		void SetGrainOscType(const EOsc::Type InGrainOscType);
+		SIGNALPROCESSING_API void SetGrainOscType(const EOsc::Type InGrainOscType);
 
 		// Sets the base grain volume and randomization range
-		void SetGrainVolume(const float InBaseVolume, const FVector2D InVolumeRange = FVector2D::ZeroVector);
+		SIGNALPROCESSING_API void SetGrainVolume(const float InBaseVolume, const FVector2D InVolumeRange = FVector2D::ZeroVector);
 
 		// Sets the grain modulation -- allows modulating actively playing grain volumes
-		void SetGrainVolumeModulation(const float InVolumeModulation);
+		SIGNALPROCESSING_API void SetGrainVolumeModulation(const float InVolumeModulation);
 
 		// Sets the base grain pitch and randomization range
-		void SetGrainPitch(const float InBasePitch, const FVector2D InPitchRange = FVector2D::ZeroVector);
+		SIGNALPROCESSING_API void SetGrainPitch(const float InBasePitch, const FVector2D InPitchRange = FVector2D::ZeroVector);
 
 		// Sets the grain frequency
-		void SetGrainFrequency(const float InFrequency, const FVector2D InPitchRange = FVector2D::ZeroVector);
+		SIGNALPROCESSING_API void SetGrainFrequency(const float InFrequency, const FVector2D InPitchRange = FVector2D::ZeroVector);
 
 		// Sets the grain frequency modulation
-		void SetGrainFrequencyModulation(const float InFrequencyModulation);
+		SIGNALPROCESSING_API void SetGrainFrequencyModulation(const float InFrequencyModulation);
 
 		// Sets the grain pitch modulation -- allows modulating actively playing grain pitches
-		void SetGrainPitchModulation(const float InPitchModulation);
+		SIGNALPROCESSING_API void SetGrainPitchModulation(const float InPitchModulation);
 
 		// Sets the grain azimuth (pan) and randomization range
-		void SetGrainPan(const float InBasePan, const FVector2D InPanRange = FVector2D::ZeroVector);
+		SIGNALPROCESSING_API void SetGrainPan(const float InBasePan, const FVector2D InPanRange = FVector2D::ZeroVector);
 
 		// Sets the grain azimuth modulation - allows modulating actively playing grain azimuths
-		void SetGrainPanModulation(const float InPanModulation);
+		SIGNALPROCESSING_API void SetGrainPanModulation(const float InPanModulation);
 
 		// Sets the grain duration. 
-		void SetGrainDuration(const float InBaseDuration, const FVector2D InDurationRange = FVector2D::ZeroVector);
+		SIGNALPROCESSING_API void SetGrainDuration(const float InBaseDuration, const FVector2D InDurationRange = FVector2D::ZeroVector);
 
 		// Sets the grain duration modulation.
-		void SetGrainDurationScale(const float InDurationScale);
+		SIGNALPROCESSING_API void SetGrainDurationScale(const float InDurationScale);
 
 		// Return the number of currently active grains
-		int32 GetNumActiveGrains() const;
+		SIGNALPROCESSING_API int32 GetNumActiveGrains() const;
 
 		// Get current playback time (in granular mode)
-		float GetCurrentPlayheadTime() const;
+		SIGNALPROCESSING_API float GetCurrentPlayheadTime() const;
 
 		// Returns the duration of the internal loaded sample buffer
-		float GetSampleDuration() const;
+		SIGNALPROCESSING_API float GetSampleDuration() const;
 
 		// Generate the next audio buffer
-		void Generate(float* OutAudiobuffer, const int32 NumFrames);
+		SIGNALPROCESSING_API void Generate(float* OutAudiobuffer, const int32 NumFrames);
 
 	protected:
 		// Spawns grains
-		void SpawnGrain();
+		SIGNALPROCESSING_API void SpawnGrain();
 
 		// Return wrapped playhead position
-		float GetWrappedPlayheadPosition(float PlayheadFrame);
+		SIGNALPROCESSING_API float GetWrappedPlayheadPosition(float PlayheadFrame);
 
 
 		// Current grain azimuth modulation
diff --git a/Engine/Source/Runtime/SignalProcessing/Public/DSP/IntegerDelay.h b/Engine/Source/Runtime/SignalProcessing/Public/DSP/IntegerDelay.h
--- a/Engine/Source/Runtime/SignalProcessing/Public/DSP/IntegerDelay.h
+++ b/Engine/Source/Runtime/SignalProcessing/Public/DSP/IntegerDelay.h
@@ -11,32 +11,32 @@ namespace Audio
 {
 
 	// An adjustable delay line. Delays values are limited to integer values. 
-	class SIGNALPROCESSING_API FIntegerDelay
+	class FIntegerDelay
 	{
 	public:
 
 		// InMaxDelaySamples is the maximum supported delay.
 		// InDelaySamples is the initial delay in samples.
-		FIntegerDelay(int32 InMaxNumDelaySamples, int32 InNumDelaySamples, int32 InNumInternalBufferSamples = 256);
+		SIGNALPROCESSING_API FIntegerDelay(int32 InMaxNumDelaySamples, int32 InNumDelaySamples, int32 InNumInternalBufferSamples = 256);
 
 		// Destructor
-		~FIntegerDelay();
+		SIGNALPROCESSING_API ~FIntegerDelay();
 
 		// Sets the current delay in samples. InDelay must be less than or equal to the InMaxDelaySamples set in the constructor.
-		void SetDelayLengthSamples(int32 InNumDelaySamples);
+		SIGNALPROCESSING_API void SetDelayLengthSamples(int32 InNumDelaySamples);
 
 		// Resets the delay line state, flushes buffer and resets read/write pointers.
-		void Reset();
+		SIGNALPROCESSING_API void Reset();
 
 		// Returns the current delay length in samples.
-		int32 GetNumDelaySamples() const;
+		SIGNALPROCESSING_API int32 GetNumDelaySamples() const;
 
 		// Process InSamples, placing delayed versions in OutSamples.
-		void ProcessAudio(const Audio::FAlignedFloatBuffer& InSamples, Audio::FAlignedFloatBuffer& OutSamples);
+		SIGNALPROCESSING_API void ProcessAudio(const Audio::FAlignedFloatBuffer& InSamples, Audio::FAlignedFloatBuffer& OutSamples);
 
 		// Retrieve a copy of the internal delay line.
 		// InNum must be less than or equal to InMaxNumDelaySamples.
-		void PeekDelayLine(int32 InNum, Audio::FAlignedFloatBuffer& OutSamples);
+		SIGNALPROCESSING_API void PeekDelayLine(int32 InNum, Audio::FAlignedFloatBuffer& OutSamples);
 
 	private:
 		// Process a block of audio. InNum is always less than or equal to NumInternalBufferSamples.
diff --git a/Engine/Source/Runtime/SignalProcessing/Public/DSP/InterpolatedLinearPitchShifter.h b/Engine/Source/Runtime/SignalProcessing/Public/DSP/InterpolatedLinearPitchShifter.h
--- a/Engine/Source/Runtime/SignalProcessing/Public/DSP/InterpolatedLinearPitchShifter.h
+++ b/Engine/Source/Runtime/SignalProcessing/Public/DSP/InterpolatedLinearPitchShifter.h
@@ -8,19 +8,19 @@
 
 namespace Audio
 {
-	class SIGNALPROCESSING_API FLinearPitchShifter
+	class FLinearPitchShifter
 	{
 	public:
 		// ctor
-		FLinearPitchShifter(int32 InNumChannels = 0);
+		SIGNALPROCESSING_API FLinearPitchShifter(int32 InNumChannels = 0);
 
-		void Reset(int32 InNumChannels, float InInitialPitchShiftSemitones = 0.0f, int32 InInterpLengthFrames = 100);
+		SIGNALPROCESSING_API void Reset(int32 InNumChannels, float InInitialPitchShiftSemitones = 0.0f, int32 InInterpLengthFrames = 100);
 
 		// Sample rate converts the input audio buffer and pushes converted audio to the output circular buffer.
 		// (Returns the number of output frames generated.)
-		int32 ProcessAudio(const TArrayView<float> InputBuffer, Audio::TCircularAudioBuffer<float>& OutputBuffer);
+		SIGNALPROCESSING_API int32 ProcessAudio(const TArrayView<float> InputBuffer, Audio::TCircularAudioBuffer<float>& OutputBuffer);
 
-		void UpdatePitchShift(float InNewPitchSemitones);
+		SIGNALPROCESSING_API void UpdatePitchShift(float InNewPitchSemitones);
 
 	private:
 		// returns next fractional index increment
diff --git a/Engine/Source/Runtime/SignalProcessing/Public/DSP/InterpolatedOnePole.h b/Engine/Source/Runtime/SignalProcessing/Public/DSP/InterpolatedOnePole.h
--- a/Engine/Source/Runtime/SignalProcessing/Public/DSP/InterpolatedOnePole.h
+++ b/Engine/Source/Runtime/SignalProcessing/Public/DSP/InterpolatedOnePole.h
@@ -7,18 +7,18 @@
 
 namespace Audio
 {
-	class SIGNALPROCESSING_API FInterpolatedLPF
+	class FInterpolatedLPF
 	{
 	public:
 		// ctor
-		FInterpolatedLPF();
+		SIGNALPROCESSING_API FInterpolatedLPF();
 
 		FORCEINLINE float GetCutoffFrequency() const
 		{
 			return CutoffFrequency;
 		}
 
-		void Init(float InSampleRate, int32 InNumChannels);
+		SIGNALPROCESSING_API void Init(float InSampleRate, int32 InNumChannels);
 
 		/*
 			StartFrequencyInterpolation():
@@ -29,13 +29,13 @@ namespace Audio
 				Failing to call StopFrequencyInterpolation() after more than InterpLength calls to ProcessAudioFrame
 				will cause the filter to become unstable and blow up to +/- inf!
 		*/
-		void StartFrequencyInterpolation(const float InTargetFrequency, const int32 InterpLength = 1);
+		SIGNALPROCESSING_API void StartFrequencyInterpolation(const float InTargetFrequency, const int32 InterpLength = 1);
 
 		// interpolates coefficient and processes a sample
-		void ProcessAudioFrame(const float* RESTRICT InputFrame, float* RESTRICT OutputFrame);
+		SIGNALPROCESSING_API void ProcessAudioFrame(const float* RESTRICT InputFrame, float* RESTRICT OutputFrame);
 
 		// interpolates coefficient and processes a buffer
-		void ProcessAudioBuffer(const float* RESTRICT InputBuffer, float* RESTRICT OutputBuffer, const int32 NumSamples);
+		SIGNALPROCESSING_API void ProcessAudioBuffer(const float* RESTRICT InputBuffer, float* RESTRICT OutputBuffer, const int32 NumSamples);
 
 		/*
 			StopFrequencyInterpolation() needs to be called manually when the interpolation should be done.
@@ -49,11 +49,11 @@ namespace Audio
 		}
 
 		// Resets the sample delay to 0
-		void Reset();
+		SIGNALPROCESSING_API void Reset();
 
 		// Clears memory without reevaluating coefficients.
 		// This function is useful when there is a break between ProcessAudio calls.
-		void ClearMemory();
+		SIGNALPROCESSING_API void ClearMemory();
 
 	private:
 		float CutoffFrequency{ -1.0f };
@@ -72,18 +72,18 @@ namespace Audio
 
 
 
-	class SIGNALPROCESSING_API FInterpolatedHPF
+	class FInterpolatedHPF
 	{
 	public:
 		// ctor
-		FInterpolatedHPF();
+		SIGNALPROCESSING_API FInterpolatedHPF();
 
 		FORCEINLINE float GetCutoffFrequency() const
 		{
 			return CutoffFrequency;
 		}
 
-		void Init(float InSampleRate, int32 InNumChannels);
+		SIGNALPROCESSING_API void Init(float InSampleRate, int32 InNumChannels);
 
 		/*
 			StartFrequencyInterpolation():
@@ -94,13 +94,13 @@ namespace Audio
 				Failing to call StopFrequencyInterpolation() after more than InterpLength calls to ProcessAudioFrame
 				will cause the filter to become unstable and blow up to +/- inf!
 		*/
-		void StartFrequencyInterpolation(const float InTargetFrequency, const int32 InterpLength = 1);
+		SIGNALPROCESSING_API void StartFrequencyInterpolation(const float InTargetFrequency, const int32 InterpLength = 1);
 
 		// interpolates coefficient and processes a sample
-		void ProcessAudioFrame(const float* RESTRICT InputFrame, float* RESTRICT OutputFrame);
+		SIGNALPROCESSING_API void ProcessAudioFrame(const float* RESTRICT InputFrame, float* RESTRICT OutputFrame);
 
 		// interpolates coefficient and processes a buffer
-		void ProcessAudioBuffer(const float* RESTRICT InputBuffer, float* RESTRICT OutputBuffer, const int32 NumSamples);
+		SIGNALPROCESSING_API void ProcessAudioBuffer(const float* RESTRICT InputBuffer, float* RESTRICT OutputBuffer, const int32 NumSamples);
 
 		/*
 			StopFrequencyInterpolation() needs to be called manually when the interpolation should be done.
@@ -114,11 +114,11 @@ namespace Audio
 		}
 
 		// Resets the sample delay to 0
-		void Reset();
+		SIGNALPROCESSING_API void Reset();
 
 		// Clears memory without reevaluating coefficients.
 		// This function is useful when there is a break between ProcessAudio calls.
-		void ClearMemory();
+		SIGNALPROCESSING_API void ClearMemory();
 		
 	private:
 		float CutoffFrequency{ -1.0f };
diff --git a/Engine/Source/Runtime/SignalProcessing/Public/DSP/LFO.h b/Engine/Source/Runtime/SignalProcessing/Public/DSP/LFO.h
--- a/Engine/Source/Runtime/SignalProcessing/Public/DSP/LFO.h
+++ b/Engine/Source/Runtime/SignalProcessing/Public/DSP/LFO.h
@@ -41,57 +41,57 @@ namespace Audio
 	}
 
 	// Low frequency oscillator
-	class SIGNALPROCESSING_API FLFO : public IOscBase
+	class FLFO : public IOscBase
 	{
 	public:
-		FLFO();
+		SIGNALPROCESSING_API FLFO();
 		virtual ~FLFO() = default;
 
 		//~ Begin FOscBase
-		virtual void Init(const float InSampleRate, const int32 InVoiceId = 0, FModulationMatrix* InMatrix = nullptr, const int32 ModMatrixStage = 0) override;
-		virtual void Start() override;
-		virtual void Stop() override;
-		virtual void Reset() override;
-		virtual float Generate(float* QuadPhaseOutput = nullptr) override;
+		SIGNALPROCESSING_API virtual void Init(const float InSampleRate, const int32 InVoiceId = 0, FModulationMatrix* InMatrix = nullptr, const int32 ModMatrixStage = 0) override;
+		SIGNALPROCESSING_API virtual void Start() override;
+		SIGNALPROCESSING_API virtual void Stop() override;
+		SIGNALPROCESSING_API virtual void Reset() override;
+		SIGNALPROCESSING_API virtual float Generate(float* QuadPhaseOutput = nullptr) override;
 		//~ End FOscBase
 
 		// Set whether or not LFO outputs in bipolar domain or unipolar (false) domain. Defaults to bipolar.
-		void SetBipolar(const bool bInBipolar);
+		SIGNALPROCESSING_API void SetBipolar(const bool bInBipolar);
 
 		// Set the waveform type of LFO generator
-		void SetType(const ELFO::Type InLFOType);
+		SIGNALPROCESSING_API void SetType(const ELFO::Type InLFOType);
 
 		// Returns current waveform type of LFO generator
-		ELFO::Type GetType() const;
+		SIGNALPROCESSING_API ELFO::Type GetType() const;
 
 		// Sets mode of LFO (Looping or one-shot.  Free not supported)
-		void SetMode(const ELFOMode::Type InLFOMode);
+		SIGNALPROCESSING_API void SetMode(const ELFOMode::Type InLFOMode);
 
 		// Returns current mode of LFO
-		ELFOMode::Type GetMode() const;
+		SIGNALPROCESSING_API ELFOMode::Type GetMode() const;
 
 		// Sets a phase offset for the LFO.  Clamped to positive values only
-		void SetPhaseOffset(const float InOffset);
+		SIGNALPROCESSING_API void SetPhaseOffset(const float InOffset);
 
 		// Sets the exponential factor of LFO if type is "exponential". Clamped to positive, non-zero values.
-		void SetExponentialFactor(const float InExpFactor);
+		SIGNALPROCESSING_API void SetExponentialFactor(const float InExpFactor);
 
 		// Returns mod source's normal phase patch source
-		FPatchSource GetModSourceNormalPhase() const;
+		SIGNALPROCESSING_API FPatchSource GetModSourceNormalPhase() const;
 
 		// Returns mod source's quad phase patch source
-		FPatchSource GetModSourceQuadPhase() const;
+		SIGNALPROCESSING_API FPatchSource GetModSourceQuadPhase() const;
 
 	protected:
-		float ComputeLFO(float InputPhase, float* OutQuad = nullptr);
+		SIGNALPROCESSING_API float ComputeLFO(float InputPhase, float* OutQuad = nullptr);
 
 		// Returns initial phase, which differs between generators
 		// ensuring certain LFO shapes start on rising edge from zero-
 		// crossing by default (assuming user's provided phase offset is 0).
-		float GetInitPhase() const;
+		SIGNALPROCESSING_API float GetInitPhase() const;
 
 		// Resets generator to initial phase
-		virtual void ResetPhase() override;
+		SIGNALPROCESSING_API virtual void ResetPhase() override;
 
 		ELFO::Type LFOType;
 		ELFOMode::Type LFOMode;
diff --git a/Engine/Source/Runtime/SignalProcessing/Public/DSP/LateReflectionsFast.h b/Engine/Source/Runtime/SignalProcessing/Public/DSP/LateReflectionsFast.h
--- a/Engine/Source/Runtime/SignalProcessing/Public/DSP/LateReflectionsFast.h
+++ b/Engine/Source/Runtime/SignalProcessing/Public/DSP/LateReflectionsFast.h
@@ -12,7 +12,7 @@
 namespace Audio
 {
 	// Structure to hold the various delay line tap outputs produced from a plate.
-	struct SIGNALPROCESSING_API FLateReflectionsPlateOutputs
+	struct FLateReflectionsPlateOutputs
 	{
 		// Number of tap locations in plate reverb
 		static const int32 NumTaps = 7;
@@ -24,12 +24,12 @@ namespace Audio
 		FAlignedFloatBuffer Output;
 
 		// Resizes and zero all buffers in structure.
-		void ResizeAndZero(int32 InNumSamples);
+		SIGNALPROCESSING_API void ResizeAndZero(int32 InNumSamples);
 	};
 
 	// Delay line settings for reverb plate delay
 	// Tap locations are determined by analyzing these delay values. 
-	struct SIGNALPROCESSING_API FLateReflectionsPlateDelays
+	struct FLateReflectionsPlateDelays
 	{
 		/*
 			Setting the delay values affects the tap locations. The audio flowchart below shows the
@@ -110,50 +110,50 @@ namespace Audio
 		int32 NumSamplesDelayI;
 
 		// Default delay settings for left channel.
-		static FLateReflectionsPlateDelays DefaultLeftDelays(float InSampleRate);
+		static SIGNALPROCESSING_API FLateReflectionsPlateDelays DefaultLeftDelays(float InSampleRate);
 
 		// Default delay settings for right channel.
-		static FLateReflectionsPlateDelays DefaultRightDelays(float InSampleRate);
+		static SIGNALPROCESSING_API FLateReflectionsPlateDelays DefaultRightDelays(float InSampleRate);
 	};
 
 
 	// Single plate channel for plate reverberation.
-	class SIGNALPROCESSING_API FLateReflectionsPlate {
+	class FLateReflectionsPlate {
 	public:
 		// InMaxNumInternalBufferSamples sets the mmaximum number of samples in an internal buffer.
-		FLateReflectionsPlate(
+		SIGNALPROCESSING_API FLateReflectionsPlate(
 			float InSampleRate,
 			int32 InMaxNumInternalBufferSamples,
 			const FLateReflectionsPlateDelays& InDelays);
 
-		~FLateReflectionsPlate();
+		SIGNALPROCESSING_API ~FLateReflectionsPlate();
 
 		// Processing input samples. All input buffers must be of equal length.
 		// InSamples contains previously unseen audio samples.
 		// InFeedbackSamples contains the delay line output from the other plate.
 		// InDelayModulations contains the fractional delay values per a sample to modulate the first all pass filters delay line.
 		// OutPlateSamples is filled with audio from the various tap points and delay line output.
-		void ProcessAudioFrames(
+		SIGNALPROCESSING_API void ProcessAudioFrames(
 			const FAlignedFloatBuffer& InSamples,
 			const FAlignedFloatBuffer& InFeedbackSamples,
 			const FAlignedFloatBuffer& InDelayModulations,
 			FLateReflectionsPlateOutputs& OutPlateSamples);
 
 		// Flush internal audio to silence.
-		void FlushAudio();
+		SIGNALPROCESSING_API void FlushAudio();
 
-		void SetDensity(float InDensity);
-		void SetDampening(float InDampening);
-		void SetDecay(float InDecay);
+		SIGNALPROCESSING_API void SetDensity(float InDensity);
+		SIGNALPROCESSING_API void SetDampening(float InDampening);
+		SIGNALPROCESSING_API void SetDecay(float InDecay);
 
 		// The number of samples in an internal buffer. This describes the limit of how many values
 		// are accessible via PeekDelayLineOutput(...)
-		int32 GetNumInternalBufferSamples() const;
+		SIGNALPROCESSING_API int32 GetNumInternalBufferSamples() const;
 
 		// Retrieve the internal delay line data. The number of samples retrieved is limited to the
 		// number of samples in the internal buffer. That amount can be determined by calling
 		// "GetNumInternalBufferSamples()".
-		void PeekDelayLine(int32 InNum, FAlignedFloatBuffer& OutSamples);
+		SIGNALPROCESSING_API void PeekDelayLine(int32 InNum, FAlignedFloatBuffer& OutSamples);
 
 	protected:
 		// Current parameter settings of reverb
@@ -188,7 +188,7 @@ namespace Audio
 	};
 
 	// Settings for controlling the FLateReflections
-	struct SIGNALPROCESSING_API FLateReflectionsFastSettings
+	struct FLateReflectionsFastSettings
 	{
 		// Milliseconds for the predelay
 		float LateDelayMsec;
@@ -211,57 +211,57 @@ namespace Audio
 		// The amount of diffusion in decay path. Larger values is a more dense reverb.
 		float Density;
 
-		FLateReflectionsFastSettings();
+		SIGNALPROCESSING_API FLateReflectionsFastSettings();
 
-		bool operator==(const FLateReflectionsFastSettings& Other) const;
+		SIGNALPROCESSING_API bool operator==(const FLateReflectionsFastSettings& Other) const;
 
-		bool operator!=(const FLateReflectionsFastSettings& Other) const;
+		SIGNALPROCESSING_API bool operator!=(const FLateReflectionsFastSettings& Other) const;
 	};
 
 
 	// FLateReflections generates the long tail reverb of an input audio signal using a relatively
 	// fast algorithm using all pass filter delay lines. It supports 1 or 2 channel input and always 
 	// produces 2 channel output. 
-	class SIGNALPROCESSING_API FLateReflectionsFast {
+	class FLateReflectionsFast {
 	public:
 
 		// Limits on settings.
-		static const float MaxLateDelayMsec;
-		static const float MinLateDelayMsec;
-		static const float MaxLateGainDB;
-		static const float MaxBandwidth;
-		static const float MinBandwidth;
-		static const float MaxDampening;
-		static const float MinDampening;
-		static const float MaxDiffusion;
-		static const float MinDiffusion;
-		static const float MaxDecay;
-		static const float MinDecay;
-		static const float MaxDensity;
-		static const float MinDensity;
-
-		static const FLateReflectionsFastSettings DefaultSettings;
+		static SIGNALPROCESSING_API const float MaxLateDelayMsec;
+		static SIGNALPROCESSING_API const float MinLateDelayMsec;
+		static SIGNALPROCESSING_API const float MaxLateGainDB;
+		static SIGNALPROCESSING_API const float MaxBandwidth;
+		static SIGNALPROCESSING_API const float MinBandwidth;
+		static SIGNALPROCESSING_API const float MaxDampening;
+		static SIGNALPROCESSING_API const float MinDampening;
+		static SIGNALPROCESSING_API const float MaxDiffusion;
+		static SIGNALPROCESSING_API const float MinDiffusion;
+		static SIGNALPROCESSING_API const float MaxDecay;
+		static SIGNALPROCESSING_API const float MinDecay;
+		static SIGNALPROCESSING_API const float MaxDensity;
+		static SIGNALPROCESSING_API const float MinDensity;
+
+		static SIGNALPROCESSING_API const FLateReflectionsFastSettings DefaultSettings;
 
 		// InMaxNumInternalBufferSamples sets the maximum possible number of samples in an internal
 		// buffer. This only affectsinternal chunking and does not affect the number of samples that
 		// can be sent to ProcessAudio(...), 
-		FLateReflectionsFast(float InSampleRate, int InMaxNumInternalBufferSamples, const FLateReflectionsFastSettings& InSettings=DefaultSettings);
+		SIGNALPROCESSING_API FLateReflectionsFast(float InSampleRate, int InMaxNumInternalBufferSamples, const FLateReflectionsFastSettings& InSettings=DefaultSettings);
 		
-		~FLateReflectionsFast();
+		SIGNALPROCESSING_API ~FLateReflectionsFast();
 
 		// Copies the reverb settings internally, clamps the internal copy and applies the clamped settings.
-		void SetSettings(const FLateReflectionsFastSettings& InSettings);
+		SIGNALPROCESSING_API void SetSettings(const FLateReflectionsFastSettings& InSettings);
 
 		// Alters settings to ensure they are within acceptable ranges. 
-		static void ClampSettings(FLateReflectionsFastSettings& InOutSettings);
+		static SIGNALPROCESSING_API void ClampSettings(FLateReflectionsFastSettings& InOutSettings);
 
 		// Create reverberation in OutSamples based upon InSamples. 
 		// OutSamples is a 2 channel audio buffer no matter the value of InNumChannels.
 		// OutSamples is mixed with InSamples based upon gain.
-		void ProcessAudio(const FAlignedFloatBuffer& InSamples, const int32 InNumChannels, FAlignedFloatBuffer& OutLeftSamples, FAlignedFloatBuffer& OutRightSamples);
+		SIGNALPROCESSING_API void ProcessAudio(const FAlignedFloatBuffer& InSamples, const int32 InNumChannels, FAlignedFloatBuffer& OutLeftSamples, FAlignedFloatBuffer& OutRightSamples);
 
 		// Flush internal audio to silence.
-		void FlushAudio();
+		SIGNALPROCESSING_API void FlushAudio();
 
 	private:
 		void ProcessAudioBuffer(const float* InSampleData, const int32 InNumFrames, const int32 InNumChannels, float* OutLeftSampleData, float* OutRightSampleData);
diff --git a/Engine/Source/Runtime/SignalProcessing/Public/DSP/LinearInterpFractionalDelay.h b/Engine/Source/Runtime/SignalProcessing/Public/DSP/LinearInterpFractionalDelay.h
--- a/Engine/Source/Runtime/SignalProcessing/Public/DSP/LinearInterpFractionalDelay.h
+++ b/Engine/Source/Runtime/SignalProcessing/Public/DSP/LinearInterpFractionalDelay.h
@@ -8,21 +8,21 @@
 namespace Audio
 {
 	// Fractional delay using linear interpolation.
-	class SIGNALPROCESSING_API FLinearInterpFractionalDelay
+	class FLinearInterpFractionalDelay
 	{
 	public:
 		
 		// InMaxDelay sets the maximum allowable delay that this object can support. The minimum delay is 0.0
 		// InNumInternalBufferSamples sets the maximum block processing size.
-		FLinearInterpFractionalDelay(int32 InMaxDelay, int32 InMaxNumInternalBufferSamples);
-		~FLinearInterpFractionalDelay();
+		SIGNALPROCESSING_API FLinearInterpFractionalDelay(int32 InMaxDelay, int32 InMaxNumInternalBufferSamples);
+		SIGNALPROCESSING_API ~FLinearInterpFractionalDelay();
 
 		// Apply delay to InSamples. Fill OutSamples with data from the delay line at a delay of InDelay.
 		// InDelay must be equal length to InSamples. 
-		void ProcessAudio(const FAlignedFloatBuffer& InSamples, const FAlignedFloatBuffer& InDelays, FAlignedFloatBuffer& OutSamples);
+		SIGNALPROCESSING_API void ProcessAudio(const FAlignedFloatBuffer& InSamples, const FAlignedFloatBuffer& InDelays, FAlignedFloatBuffer& OutSamples);
 
 		// Set all values in internal delay line to zero. 
-		void Reset();
+		SIGNALPROCESSING_API void Reset();
 
 	private:
 		void ProcessAudioBlock(const float* InSamples, const float* InDelays, const int32 InNum, float* OutSamples);
diff --git a/Engine/Source/Runtime/SignalProcessing/Public/DSP/LinkwitzRileyBandSplitter.h b/Engine/Source/Runtime/SignalProcessing/Public/DSP/LinkwitzRileyBandSplitter.h
--- a/Engine/Source/Runtime/SignalProcessing/Public/DSP/LinkwitzRileyBandSplitter.h
+++ b/Engine/Source/Runtime/SignalProcessing/Public/DSP/LinkwitzRileyBandSplitter.h
@@ -68,7 +68,7 @@ namespace Audio
 	 * Helper for Multi-Band processing to generate Linwitz-Riley filtered outputs from input
 	 * https://en.wikipedia.org/wiki/Linkwitz%E2%80%93Riley_filter
 	 */
-	class SIGNALPROCESSING_API FLinkwitzRileyBandSplitter
+	class FLinkwitzRileyBandSplitter
 	{
 	public:
 		struct FCrossoverBandwidthPair
@@ -80,16 +80,16 @@ namespace Audio
 		FLinkwitzRileyBandSplitter() {};
 
 		// initalize filters
-		void Init(const int32 InChannels,
+		SIGNALPROCESSING_API void Init(const int32 InChannels,
 				  const float InSampleRate,
 				  const EFilterOrder FilterOrder, 
 				  const bool bInPhaseCompensate,
 				  const TArray<float>& InCrossovers); // Always InBands - 1 Crossovers
 
-		void ProcessAudioFrame(const float* InBuffer, FMultibandBuffer& OutBuffer);
-		void ProcessAudioBuffer(const float* InBuffer, FMultibandBuffer& OutBuffer, const int32 NumFrames);
+		SIGNALPROCESSING_API void ProcessAudioFrame(const float* InBuffer, FMultibandBuffer& OutBuffer);
+		SIGNALPROCESSING_API void ProcessAudioBuffer(const float* InBuffer, FMultibandBuffer& OutBuffer, const int32 NumFrames);
 
-		void SetCrossovers(const TArray<float>& InCrossoverFrequencies);
+		SIGNALPROCESSING_API void SetCrossovers(const TArray<float>& InCrossoverFrequencies);
 
 	private:
 		EFilterOrder FilterOrder = EFilterOrder::FourPole;
@@ -106,8 +106,8 @@ namespace Audio
 		TArray<FLinkwitzRileyBandFilter> BandFilters;
 		TArray<FCrossoverBandwidthPair> Crossovers;
 
-		void CopyToBuffer(float* Destination, const float* Origin, const int32 NumSamples);
-		void InvertBuffer(float* Buffer, const int32 NumSamples);
-		float GetQ(EFilterOrder InFilterOrder);
+		SIGNALPROCESSING_API void CopyToBuffer(float* Destination, const float* Origin, const int32 NumSamples);
+		SIGNALPROCESSING_API void InvertBuffer(float* Buffer, const int32 NumSamples);
+		SIGNALPROCESSING_API float GetQ(EFilterOrder InFilterOrder);
 	};
 }
diff --git a/Engine/Source/Runtime/SignalProcessing/Public/DSP/LongDelayAPF.h b/Engine/Source/Runtime/SignalProcessing/Public/DSP/LongDelayAPF.h
--- a/Engine/Source/Runtime/SignalProcessing/Public/DSP/LongDelayAPF.h
+++ b/Engine/Source/Runtime/SignalProcessing/Public/DSP/LongDelayAPF.h
@@ -10,39 +10,39 @@ namespace Audio
 {
 	// All Pass Filter with a long delay. This filter is specifically designed for 
 	// reverb applications where filter delay lines are long.
-	class SIGNALPROCESSING_API FLongDelayAPF
+	class FLongDelayAPF
 	{
 	public:
 		// InG is the filter coefficient used in the long delay all pass filter.
 		// InNumDelaySamples is the delay line length in samples.
 		// InMaxNumInternalBufferSamples is the maximum internal block size used internally. 
-		FLongDelayAPF(float InG, int32 InNumDelaySamples, int32 InMaxNumInternalBufferSamples);
+		SIGNALPROCESSING_API FLongDelayAPF(float InG, int32 InNumDelaySamples, int32 InMaxNumInternalBufferSamples);
 
 		// Destructor
-		~FLongDelayAPF();
+		SIGNALPROCESSING_API ~FLongDelayAPF();
 
 		// Set the APF feedback/feedforward gain coefficient
 		void SetG(float InG) { G = InG; }
 
 		// Process InSamples and place filtered data in OutSamples
-		void ProcessAudio(const FAlignedFloatBuffer& InSamples, FAlignedFloatBuffer& OutSamples);
+		SIGNALPROCESSING_API void ProcessAudio(const FAlignedFloatBuffer& InSamples, FAlignedFloatBuffer& OutSamples);
 
 		// Process Samples in place
-		void ProcessAudio(FAlignedFloatBuffer& Samples);
+		SIGNALPROCESSING_API void ProcessAudio(FAlignedFloatBuffer& Samples);
 
 		// Process InSamples and place filtered data in OutSamples and delay line samples in OutDelaySamples.
 		// The OutDelaySamples correspond to "w[n] = InSamples[n] + InG * w[n - InDelay]"
-		void ProcessAudio(const FAlignedFloatBuffer& InSamples, FAlignedFloatBuffer& OutSamples, FAlignedFloatBuffer& OutDelaySamples);
+		SIGNALPROCESSING_API void ProcessAudio(const FAlignedFloatBuffer& InSamples, FAlignedFloatBuffer& OutSamples, FAlignedFloatBuffer& OutDelaySamples);
 
 		// Sets delay line values to zero.
-		void Reset();
+		SIGNALPROCESSING_API void Reset();
 
 	private:
 		// Performs all pass filtering on an audio block of BlockSize samples are less. 
 		// InNum must be less than or equal to BlockSize.
-		void ProcessAudioBlock(const float* InSamples, const float* InDelaySamples, const int32 InNum, float* OutSamples, float* OutDelaySamples);
+		SIGNALPROCESSING_API void ProcessAudioBlock(const float* InSamples, const float* InDelaySamples, const int32 InNum, float* OutSamples, float* OutDelaySamples);
 
-		int32 GetNumInternalBufferSamples() const;
+		SIGNALPROCESSING_API int32 GetNumInternalBufferSamples() const;
 
 		// Feedback/Feedforward gain coefficient
 		float G;
diff --git a/Engine/Source/Runtime/SignalProcessing/Public/DSP/MelScale.h b/Engine/Source/Runtime/SignalProcessing/Public/DSP/MelScale.h
--- a/Engine/Source/Runtime/SignalProcessing/Public/DSP/MelScale.h
+++ b/Engine/Source/Runtime/SignalProcessing/Public/DSP/MelScale.h
@@ -28,7 +28,7 @@ namespace Audio
 
 	/** Settings for a mel kernel which transforms an linearly space spectrum (e.g. FFT Magnitude)
 	 * to a mel spectrum */
-	struct SIGNALPROCESSING_API FMelSpectrumKernelSettings
+	struct FMelSpectrumKernelSettings
 	{
 		int32 NumBands;						//< Number of bands in Mel spectrum
 		float KernelMinCenterFreq;			//< Minimum frequency of lowest mel band.
diff --git a/Engine/Source/Runtime/SignalProcessing/Public/DSP/MidiNoteQuantizer.h b/Engine/Source/Runtime/SignalProcessing/Public/DSP/MidiNoteQuantizer.h
--- a/Engine/Source/Runtime/SignalProcessing/Public/DSP/MidiNoteQuantizer.h
+++ b/Engine/Source/Runtime/SignalProcessing/Public/DSP/MidiNoteQuantizer.h
@@ -62,31 +62,31 @@ namespace Audio
 		};
 	}
 
-	struct SIGNALPROCESSING_API ScaleDegreeSet
+	struct ScaleDegreeSet
 	{
 	public:
 		// ctor
-		ScaleDegreeSet(const TArray<float>& InScaleDegrees, const TArray<float>& InChordTones = {});
+		SIGNALPROCESSING_API ScaleDegreeSet(const TArray<float>& InScaleDegrees, const TArray<float>& InChordTones = {});
 
-		TArrayView<float> GetScaleDegreeSet(bool bChordTonesOnlyIfApplicable = false);
+		SIGNALPROCESSING_API TArrayView<float> GetScaleDegreeSet(bool bChordTonesOnlyIfApplicable = false);
 
 	private:
 		TArray<float> ScaleDegrees;
 		TArray<float> ChordTones;
 	};
 
-	class SIGNALPROCESSING_API FMidiNoteQuantizer
+	class FMidiNoteQuantizer
 	{
 	public:
-		static float QuantizeMidiNote(const float InNote, const float InRoot, EMusicalScale::Scale InScale, bool bChordTonesOnlyIfApplicable = false);
-		static float QuantizeMidiNote(const float InNote, const float InRoot, const TArrayView<float> InScaleDegrees, const float InSemitoneScaleRange = 12.0f);
+		static SIGNALPROCESSING_API float QuantizeMidiNote(const float InNote, const float InRoot, EMusicalScale::Scale InScale, bool bChordTonesOnlyIfApplicable = false);
+		static SIGNALPROCESSING_API float QuantizeMidiNote(const float InNote, const float InRoot, const TArrayView<float> InScaleDegrees, const float InSemitoneScaleRange = 12.0f);
 
 	private:
 		static float QuantizeValueToScaleDegree(const float InValue, const TArrayView<float> InScaleDegrees, const float InSemitoneScaleRange = 12.0f);
 
 	public:
 		// Statically-defined scale/chord tone definitions.
-		static TMap<EMusicalScale::Scale, ScaleDegreeSet> ScaleDegreeSetMap;
+		static SIGNALPROCESSING_API TMap<EMusicalScale::Scale, ScaleDegreeSet> ScaleDegreeSetMap;
 
 	}; // class FMidNoteQuantizer
-} // namespace Audio
\ No newline at end of file
+} // namespace Audio
