[Engine]

diff --git a/Engine/Source/Runtime/SignalProcessing/Public/DSP/ModulationMatrix.h b/Engine/Source/Runtime/SignalProcessing/Public/DSP/ModulationMatrix.h
--- a/Engine/Source/Runtime/SignalProcessing/Public/DSP/ModulationMatrix.h
+++ b/Engine/Source/Runtime/SignalProcessing/Public/DSP/ModulationMatrix.h
@@ -98,49 +98,49 @@ namespace Audio
 		bool bEnabled;
 	};
 
-	class SIGNALPROCESSING_API FModulationMatrix
+	class FModulationMatrix
 	{
 	public:
-		FModulationMatrix();
-		virtual ~FModulationMatrix();
+		SIGNALPROCESSING_API FModulationMatrix();
+		SIGNALPROCESSING_API virtual ~FModulationMatrix();
 
 		// Initialize the modulation matrix with the desired number of voices
-		void Init(const int32 NumVoices);
+		SIGNALPROCESSING_API void Init(const int32 NumVoices);
 
 		// Returns the number of patch connections
-		int32 GetNumPatches(const int32 VoiceId) const;
+		SIGNALPROCESSING_API int32 GetNumPatches(const int32 VoiceId) const;
 
 		// Creates a new patch source object and returns the patch source id
-		FPatchSource CreatePatchSource(const int32 VoiceId);
+		SIGNALPROCESSING_API FPatchSource CreatePatchSource(const int32 VoiceId);
 
 		// Crates a new patch destination object and returns the patch destination id
-		FPatchDestination CreatePatchDestination(const int32 VoiceId, const int32 Stage, const float DefaultDepth);
+		SIGNALPROCESSING_API FPatchDestination CreatePatchDestination(const int32 VoiceId, const int32 Stage, const float DefaultDepth);
 
 		// Adds a new patch connection between one source and one or more destinations
-		bool AddPatch(const int32 VoiceId, FPatch* Patch);
+		SIGNALPROCESSING_API bool AddPatch(const int32 VoiceId, FPatch* Patch);
 
 		// Removes the given patch connection between a source and one or more destinations
-		bool RemovePatch(const int32 VoiceId, FPatch* Patch);
+		SIGNALPROCESSING_API bool RemovePatch(const int32 VoiceId, FPatch* Patch);
 
 		// Reset all patch state
-		void ResetPatchSourceState();
+		SIGNALPROCESSING_API void ResetPatchSourceState();
 
 		// Clear all patch connections
-		void ClearPatches(const int32 VoiceId);
+		SIGNALPROCESSING_API void ClearPatches(const int32 VoiceId);
 
 		// Set the modulation source value of the given id
-		bool SetSourceValue(const int32 VoiceId, const FPatchSource& Source, const float Value);
+		SIGNALPROCESSING_API bool SetSourceValue(const int32 VoiceId, const FPatchSource& Source, const float Value);
 
 		// Get the modulation destination value of the given destination id
-		bool GetDestinationValue(const int32 VoiceId, const FPatchDestination& Destination, float& OutValue) const;
+		SIGNALPROCESSING_API bool GetDestinationValue(const int32 VoiceId, const FPatchDestination& Destination, float& OutValue) const;
 
 		// Perform the matrix update. Optionally do only a given stage. 
-		void Update(const int32 VoiceId, const int32 Stage = INDEX_NONE);
+		SIGNALPROCESSING_API void Update(const int32 VoiceId, const int32 Stage = INDEX_NONE);
 
 	protected:
 
-		void ResetDestinations(const int32 VoiceId);
-		bool ValidatePatch(const int32 VoiceId, FPatch* Patch);
+		SIGNALPROCESSING_API void ResetDestinations(const int32 VoiceId);
+		SIGNALPROCESSING_API bool ValidatePatch(const int32 VoiceId, FPatch* Patch);
 
 		int32 NumVoices;
 		TArray<TArray<FPatch*>> Patches;
diff --git a/Engine/Source/Runtime/SignalProcessing/Public/DSP/MultichannelLinearResampler.h b/Engine/Source/Runtime/SignalProcessing/Public/DSP/MultichannelLinearResampler.h
--- a/Engine/Source/Runtime/SignalProcessing/Public/DSP/MultichannelLinearResampler.h
+++ b/Engine/Source/Runtime/SignalProcessing/Public/DSP/MultichannelLinearResampler.h
@@ -9,7 +9,7 @@
 namespace Audio
 {
 	/** Linear resampler working on multichannel buffers */
-	class SIGNALPROCESSING_API FMultichannelLinearResampler
+	class FMultichannelLinearResampler
 	{
 		static constexpr float MinFrameRatioFrameDelta = 0.001f;
 	public:
@@ -17,7 +17,7 @@ namespace Audio
 		 *
 		 * @param InNumChannel - Number of audio channels in input and output buffers.
 		 */
-		FMultichannelLinearResampler(int32 InNumChannels);
+		SIGNALPROCESSING_API FMultichannelLinearResampler(int32 InNumChannels);
 
 		/** Sets the number of input frames to read per an output frame. 0.5 is 
 		 * half speed, 1.f is normal speed, 2.0 is double speed.
@@ -26,7 +26,7 @@ namespace Audio
 		 * @param InNumFramesToInterpolate - Number of output frames over which 
 		 *                                   to interpolate the frame ratio.
 		 */
-		void SetFrameRatio(float InRatio, int32 InNumFramesToInterpolate=0);
+		SIGNALPROCESSING_API void SetFrameRatio(float InRatio, int32 InNumFramesToInterpolate=0);
 
 		/** Translates an input frame index to an output frame index given the
 		 * state of the resampler. 
@@ -37,7 +37,7 @@ namespace Audio
 		 * @param InInputFrameIndex - Index of input frame.
 		 * @return Index of output frame. 
 		 */
-		float MapInputFrameToOutputFrame(float InInputFrameIndex) const;
+		SIGNALPROCESSING_API float MapInputFrameToOutputFrame(float InInputFrameIndex) const;
 
 		/** Translates an output frame index to an input frame index given the
 		 * state of the resampler. 
@@ -48,7 +48,7 @@ namespace Audio
 		 * @param InOutputFrameIndex - Index of output frame.
 		 * @return Index of input frame. 
 		 */
-		float MapOutputFrameToInputFrame(float InOutputFrameIndex) const;
+		SIGNALPROCESSING_API float MapOutputFrameToInputFrame(float InOutputFrameIndex) const;
 
 		/** Returns the minimum number of input frames needed to produce the desired
 		 * number of output frames given the current state of the resampler. 
@@ -56,7 +56,7 @@ namespace Audio
 		 * @param InNumOutputFrames - The desired number of output frames.
 		 * @return The minimum number of input frames. 
 		 */
-		int32 GetNumInputFramesNeededToProduceOutputFrames(int32 InNumOutputFrames) const;
+		SIGNALPROCESSING_API int32 GetNumInputFramesNeededToProduceOutputFrames(int32 InNumOutputFrames) const;
 
 		/** Consumes audio from the input buffer and produces audio in the output buffer.
 		 * The desired number of frames to produce is determined by the output audio buffer
@@ -69,8 +69,8 @@ namespace Audio
 		 *
 		 * @return Actual number of frames produced. 
 		 */
-		int32 ProcessAndConsumeAudio(FMultichannelCircularBuffer& InAudio, FMultichannelBuffer& OutAudio);
-		int32 ProcessAndConsumeAudio(FMultichannelCircularBuffer& InAudio, FMultichannelBufferView& OutAudio);
+		SIGNALPROCESSING_API int32 ProcessAndConsumeAudio(FMultichannelCircularBuffer& InAudio, FMultichannelBuffer& OutAudio);
+		SIGNALPROCESSING_API int32 ProcessAndConsumeAudio(FMultichannelCircularBuffer& InAudio, FMultichannelBufferView& OutAudio);
 
 	private:
 
diff --git a/Engine/Source/Runtime/SignalProcessing/Public/DSP/MultithreadedPatching.h b/Engine/Source/Runtime/SignalProcessing/Public/DSP/MultithreadedPatching.h
--- a/Engine/Source/Runtime/SignalProcessing/Public/DSP/MultithreadedPatching.h
+++ b/Engine/Source/Runtime/SignalProcessing/Public/DSP/MultithreadedPatching.h
@@ -29,37 +29,37 @@ namespace Audio
 	 * NewOutput->PopAudio(OutAudioBufferPtr, AudioBufferNumSamples);
 	 * 
 	 */
-	struct SIGNALPROCESSING_API FPatchOutput
+	struct FPatchOutput
 	{
 	public:
-		FPatchOutput(int32 InMaxCapacity, float InGain = 1.0f);
+		SIGNALPROCESSING_API FPatchOutput(int32 InMaxCapacity, float InGain = 1.0f);
 
 		/** The default constructor will result in an uninitialized, disconnected patch point. */
-		FPatchOutput();
+		SIGNALPROCESSING_API FPatchOutput();
 
-		virtual ~FPatchOutput();
+		SIGNALPROCESSING_API virtual ~FPatchOutput();
 		
 		/** Copies the minimum of NumSamples or however many samples are available into OutBuffer. Returns the number of samples copied, or -1 if this output's corresponding input has been destroyed. */
-		int32 PopAudio(float* OutBuffer, int32 NumSamples, bool bUseLatestAudio);
+		SIGNALPROCESSING_API int32 PopAudio(float* OutBuffer, int32 NumSamples, bool bUseLatestAudio);
 
 		/** Sums the minimum of NumSamples or however many samples are available into OutBuffer. Returns the number of samples summed into OutBuffer. */
-		int32 MixInAudio(float* OutBuffer, int32 NumSamples, bool bUseLatestAudio);
+		SIGNALPROCESSING_API int32 MixInAudio(float* OutBuffer, int32 NumSamples, bool bUseLatestAudio);
 
 		/** Returns the current number of samples buffered on this output. */
-		int32 GetNumSamplesAvailable() const;
+		SIGNALPROCESSING_API int32 GetNumSamplesAvailable() const;
 
 		/** Pauses the current thread until there are the given number of samples available to pop. Will return true if it succeeded, false if it timed out. */
-		bool WaitUntilNumSamplesAvailable(int32 NumSamples, uint32 TimeOutMilliseconds = MAX_uint32);
+		SIGNALPROCESSING_API bool WaitUntilNumSamplesAvailable(int32 NumSamples, uint32 TimeOutMilliseconds = MAX_uint32);
 
 		/** Returns true if the input for this patch has been destroyed. */
-		bool IsInputStale() const;
+		SIGNALPROCESSING_API bool IsInputStale() const;
 
 		friend class FPatchInput;
 		friend class FPatchMixer;
 		friend class FPatchSplitter;
 	private:
 		
-		int32 PushAudioToInternalBuffer(const float* InBuffer, int32 NumSamples);
+		SIGNALPROCESSING_API int32 PushAudioToInternalBuffer(const float* InBuffer, int32 NumSamples);
 		
 		// Internal buffer.
 		TCircularAudioBuffer<float> InternalBuffer;
@@ -80,7 +80,7 @@ namespace Audio
 		// Event to pause the current thread until a given number of samples has been filled
 		std::atomic<FEvent*> SamplesPushedEvent;
 
-		static TAtomic<int32> PatchIDCounter;
+		static SIGNALPROCESSING_API TAtomic<int32> PatchIDCounter;
 	};
 
 	/** Patch outputs are owned by the FPatchMixer, and are pinned by the FPatchInput. */
@@ -90,36 +90,36 @@ namespace Audio
 	/**
 	 * Handle to a patch. Should only be used by a single thread.
 	 */
-	class SIGNALPROCESSING_API FPatchInput
+	class FPatchInput
 	{
 	public:
 		/** PatchInputs can only be created from explicit outputs. */
-		FPatchInput(const FPatchOutputStrongPtr& InOutput);
-		FPatchInput(const FPatchInput& Other);
-		FPatchInput& operator=(const FPatchInput& Other);
-		FPatchInput& operator=(FPatchInput&& Other);
-		FPatchInput(FPatchInput&& Other);
+		SIGNALPROCESSING_API FPatchInput(const FPatchOutputStrongPtr& InOutput);
+		SIGNALPROCESSING_API FPatchInput(const FPatchInput& Other);
+		SIGNALPROCESSING_API FPatchInput& operator=(const FPatchInput& Other);
+		SIGNALPROCESSING_API FPatchInput& operator=(FPatchInput&& Other);
+		SIGNALPROCESSING_API FPatchInput(FPatchInput&& Other);
 
 		/** Default constructed FPatchInput instances will always return -1 for PushAudio and true for IsOutputStillActive. */
 		FPatchInput() = default;
 
-		~FPatchInput();
+		SIGNALPROCESSING_API ~FPatchInput();
 
 		/** Pushes audio from InBuffer to the corresponding FPatchOutput.
 		 *  Pushes zeros if InBuffer is nullptr.
 		 *  Returns how many samples were able to be pushed, or -1 if the output was disconnected.
 		 */
-		int32 PushAudio(const float* InBuffer, int32 NumSamples);
+		SIGNALPROCESSING_API int32 PushAudio(const float* InBuffer, int32 NumSamples);
 
-		void SetGain(float InGain);
+		SIGNALPROCESSING_API void SetGain(float InGain);
 
 		/** Returns false if this output was removed, either because someone called FPatchMixer::RemoveTap with this FPatchInput, or the FPatchMixer was destroyed. */
-		bool IsOutputStillActive() const;
+		SIGNALPROCESSING_API bool IsOutputStillActive() const;
 
 		/** Returns false if this output was not initialized properly. */
-		bool IsValid() const;
+		SIGNALPROCESSING_API bool IsValid() const;
 
-		void Reset();
+		SIGNALPROCESSING_API void Reset();
 		
 		friend class FPatchMixer;
 		friend class FPatchSplitter;
@@ -137,32 +137,32 @@ namespace Audio
 	 * Important to note that this is MPSC: while multiple threads can enqueue audio on an instance of FPatchMixer using instances of FPatchInput,
 	 * only one thread can call PopAudio safely.
 	 */
-	class SIGNALPROCESSING_API FPatchMixer
+	class FPatchMixer
 	{
 	public:
 		/** Adds a new input to the tap collector. Calling this is thread safe, but individual instances of FPatchInput are only safe to be used from one thread. */
-		FPatchInput AddNewInput(int32 MaxLatencyInSamples, float InGain);
+		SIGNALPROCESSING_API FPatchInput AddNewInput(int32 MaxLatencyInSamples, float InGain);
 
 		/** Adds an existing patch input to the patch mixer. */
-		void AddNewInput(const FPatchInput& InPatchInput);
+		SIGNALPROCESSING_API void AddNewInput(const FPatchInput& InPatchInput);
 
 		/** Removes a tap from the tap collector. Calling this is thread safe, though FPatchOutput will likely not be deleted until the next call of PopAudio. */
-		void RemovePatch(const FPatchInput& InPatchInput);
+		SIGNALPROCESSING_API void RemovePatch(const FPatchInput& InPatchInput);
 
 		/** Mixes all inputs into a single buffer. This should only be called from a single thread. Returns the number of non-silent samples popped to OutBuffer. */
-		int32 PopAudio(float* OutBuffer, int32 OutNumSamples, bool bUseLatestAudio);
+		SIGNALPROCESSING_API int32 PopAudio(float* OutBuffer, int32 OutNumSamples, bool bUseLatestAudio);
 
 		/** This returns the number of inputs currently connected to this patch mixer. Thread safe, but blocks for PopAudio. */
-		int32 Num();
+		SIGNALPROCESSING_API int32 Num();
 
 		/** This function call gets the maximum number of samples that's safe to pop, based on the thread with the least amount of samples buffered. Thread safe, but blocks for PopAudio. */
-		int32 MaxNumberOfSamplesThatCanBePopped();
+		SIGNALPROCESSING_API int32 MaxNumberOfSamplesThatCanBePopped();
 
 		/** Pauses the current thread until there are the given number of samples available to pop. Will return true if it succeeded, false if it timed out. */
-		bool WaitUntilNumSamplesAvailable(int32 NumSamples, uint32 TimeOutMilliseconds = MAX_uint32);
+		SIGNALPROCESSING_API bool WaitUntilNumSamplesAvailable(int32 NumSamples, uint32 TimeOutMilliseconds = MAX_uint32);
 
 		/** Disconnect everything currently connected to this mixer. */
-		void DisconnectAllInputs();
+		SIGNALPROCESSING_API void DisconnectAllInputs();
 
 	private:
 		/** Called within PopAudio. Flushes the PendingNewPatches array into CurrentPatches. During this function, AddNewPatch is blocked. */
@@ -192,27 +192,27 @@ namespace Audio
 	 * This class is SPMC: multiple threads can call FPatchOutputStrongPtr->PopAudio safely,
 	 * but only one thread can call PushAudio.
 	 */
-	class SIGNALPROCESSING_API FPatchSplitter
+	class FPatchSplitter
 	{
 	public:
 		/**
 		 * Adds a new output. Calling this is thread safe, but individual instances of FPatchOutput are only safe to be used from one thread.
 		 * the returned FPatchOutputPtr can be safely destroyed at any point.
 		 */
-		FPatchOutputStrongPtr AddNewPatch(int32 MaxLatencyInSamples, float InGain);
+		SIGNALPROCESSING_API FPatchOutputStrongPtr AddNewPatch(int32 MaxLatencyInSamples, float InGain);
 
 		/** Adds a new a patch from an existing patch output. */
-		void AddNewPatch(FPatchOutputStrongPtr&& InPatchOutputStrongPtr);
-		void AddNewPatch(const FPatchOutputStrongPtr& InPatchOutputStrongPtr);
+		SIGNALPROCESSING_API void AddNewPatch(FPatchOutputStrongPtr&& InPatchOutputStrongPtr);
+		SIGNALPROCESSING_API void AddNewPatch(const FPatchOutputStrongPtr& InPatchOutputStrongPtr);
 
 		/** This call pushes audio to all outputs connected to this splitter. Only should be called from one thread. */
-		int32 PushAudio(const float* InBuffer, int32 InNumSamples);
+		SIGNALPROCESSING_API int32 PushAudio(const float* InBuffer, int32 InNumSamples);
 
 		/** This returns the number of outputs currently connected to this patch splitter. Thread safe, but blocks for PushAudio. */
-		int32 Num();
+		SIGNALPROCESSING_API int32 Num();
 
 		/** This function call gets the maximum number of samples that's safe to push. Thread safe, but blocks for PushAudio. */
-		int32 MaxNumberOfSamplesThatCanBePushed();
+		SIGNALPROCESSING_API int32 MaxNumberOfSamplesThatCanBePushed();
 
 	private:
 		void AddPendingPatches();
@@ -228,29 +228,29 @@ namespace Audio
 	 * This class is used to mix multiple inputs from disparate threads to a single mixdown and deliver that mixdown to multiple outputs.
 	 * This class is MPMC, but only one thread can and should call ProcessAudio().
 	 */
-	class SIGNALPROCESSING_API FPatchMixerSplitter
+	class FPatchMixerSplitter
 	{
 	public:
 		/**
 		 * Adds a new output. Calling this is thread safe, but individual instances of FPatchOutput are only safe to be used from one thread.
 		 * the returned FPatchOutputPtr can be safely destroyed at any point.
 		 */
-		FPatchOutputStrongPtr AddNewOutput(int32 MaxLatencyInSamples, float InGain);
+		SIGNALPROCESSING_API FPatchOutputStrongPtr AddNewOutput(int32 MaxLatencyInSamples, float InGain);
 
 		/** Adds a new a patch from an existing patch output. */
-		void AddNewOutput(const FPatchOutputStrongPtr& InPatchOutputStrongPtr);
+		SIGNALPROCESSING_API void AddNewOutput(const FPatchOutputStrongPtr& InPatchOutputStrongPtr);
 
 		/** Adds a new input to the tap collector. Calling this is thread safe, but individual instances of FPatchInput are only safe to be used from one thread. */
-		FPatchInput AddNewInput(int32 MaxLatencyInSamples, float InGain);
+		SIGNALPROCESSING_API FPatchInput AddNewInput(int32 MaxLatencyInSamples, float InGain);
 
 		/** Adds a new a patch input from an existing patch input object. */
-		void AddNewInput(FPatchInput& InInput);
+		SIGNALPROCESSING_API void AddNewInput(FPatchInput& InInput);
 
 		/** Removes a tap from the tap collector. Calling this is thread safe, though FPatchOutput will likely not be deleted until the next call of PopAudio. */
-		void RemovePatch(const FPatchInput& InInput);
+		SIGNALPROCESSING_API void RemovePatch(const FPatchInput& InInput);
 
 		/** Mixes audio from all inputs and pushes it to all outputs. Should be called regularly. */
-		void ProcessAudio();
+		SIGNALPROCESSING_API void ProcessAudio();
 
 	private:
 		FPatchMixer Mixer;
diff --git a/Engine/Source/Runtime/SignalProcessing/Public/DSP/Noise.h b/Engine/Source/Runtime/SignalProcessing/Public/DSP/Noise.h
--- a/Engine/Source/Runtime/SignalProcessing/Public/DSP/Noise.h
+++ b/Engine/Source/Runtime/SignalProcessing/Public/DSP/Noise.h
@@ -11,11 +11,11 @@ namespace Audio
 	* White noise generator 
 	* Flat spectrum
 	*/
-	class SIGNALPROCESSING_API FWhiteNoise
+	class FWhiteNoise
 	{
 	public:
-		FWhiteNoise();
-		FWhiteNoise(int32 InRandomSeed);
+		SIGNALPROCESSING_API FWhiteNoise();
+		SIGNALPROCESSING_API FWhiteNoise(int32 InRandomSeed);
 
 		/** Generate next sample of white noise */
 		FORCEINLINE float Generate()
@@ -36,17 +36,17 @@ namespace Audio
 	* Pink noise generator
 	* 1/Frequency noise spectrum
 	*/
-	class SIGNALPROCESSING_API FPinkNoise
+	class FPinkNoise
 	{
 	public:
 		/** Constructor. Without seed argument, uses Cpu cycles to chose one at "random" */
-		FPinkNoise();
+		SIGNALPROCESSING_API FPinkNoise();
 
 		/** Constructor with seed input */
-		FPinkNoise(int32 InRandomSeed);
+		SIGNALPROCESSING_API FPinkNoise(int32 InRandomSeed);
 
 		/** Generate next sample of pink noise. */
-		float Generate();
+		SIGNALPROCESSING_API float Generate();
 
 		/** Set Pink Noise Filter Gain (default -3db) */
 		void SetFilterGain(float InFilterGain) 
diff --git a/Engine/Source/Runtime/SignalProcessing/Public/DSP/OnePole.h b/Engine/Source/Runtime/SignalProcessing/Public/DSP/OnePole.h
--- a/Engine/Source/Runtime/SignalProcessing/Public/DSP/OnePole.h
+++ b/Engine/Source/Runtime/SignalProcessing/Public/DSP/OnePole.h
@@ -8,7 +8,7 @@
 namespace Audio
 {
 	// Simple 1-pole lowpass filter
-	class SIGNALPROCESSING_API FOnePoleLPF
+	class FOnePoleLPF
 	{
 	public:
 
diff --git a/Engine/Source/Runtime/SignalProcessing/Public/DSP/Osc.h b/Engine/Source/Runtime/SignalProcessing/Public/DSP/Osc.h
--- a/Engine/Source/Runtime/SignalProcessing/Public/DSP/Osc.h
+++ b/Engine/Source/Runtime/SignalProcessing/Public/DSP/Osc.h
@@ -49,14 +49,14 @@ namespace Audio
 	};
 
 	// Oscillator base class
-	class SIGNALPROCESSING_API IOscBase
+	class IOscBase
 	{
 	public:
-		IOscBase();
-		virtual ~IOscBase();
+		SIGNALPROCESSING_API IOscBase();
+		SIGNALPROCESSING_API virtual ~IOscBase();
 
 		// Initializes the oscillator
-		virtual void Init(const float InSampleRate, const int32 InVoiceId = 0, FModulationMatrix* InMatrix = nullptr, const int32 ModMatrixStage = 1);
+		SIGNALPROCESSING_API virtual void Init(const float InSampleRate, const int32 InVoiceId = 0, FModulationMatrix* InMatrix = nullptr, const int32 ModMatrixStage = 1);
 
 		// Starts the oscillator
 		virtual void Start() = 0;
@@ -68,13 +68,13 @@ namespace Audio
 		virtual float Generate(float* AuxOutput = nullptr) = 0;
 
 		// Updates oscillator state
-		virtual void Update();
+		SIGNALPROCESSING_API virtual void Update();
 
 		// Resets the oscillator
-		virtual void Reset();
+		SIGNALPROCESSING_API virtual void Reset();
 
 		// Resets the phase of this oscillator to 0.0f
-		virtual void ResetPhase();
+		SIGNALPROCESSING_API virtual void ResetPhase();
 
 		// Sets the gain of the oscillator
 		void SetGain(const float InGain) { Gain = InGain; }
@@ -83,13 +83,13 @@ namespace Audio
 		void SetGainMod(const float InGainMod) { ExternalGainMod = InGainMod; }
 
 		// Sets the base frequency of the oscillator
-		void SetFrequency(const float InFreqBase);
+		SIGNALPROCESSING_API void SetFrequency(const float InFreqBase);
 
 		// Sets a frequency modulation value
-		void SetFrequencyMod(const float InFreqMod);
+		SIGNALPROCESSING_API void SetFrequencyMod(const float InFreqMod);
 
 		// Sets the base frequency of the oscillator from the midi note number
-		void SetNote(const float InNote);
+		SIGNALPROCESSING_API void SetNote(const float InNote);
 
 		// Returns the frequency of the oscillator
 		float GetFrequency() const { return BaseFreq; }
@@ -97,16 +97,16 @@ namespace Audio
 		// Returns the frequency of the oscillator
 		float GetGain() const { return Gain; }
 
-		void SetCents(const float InCents);
-		void SetOctave(const float InOctave);
-		void SetSampleRate(const float InSampleRate);
-		void SetSemitones(const float InSemiTone);
-		void SetDetune(const float InDetune);
-		void SetPitchBend(const float InPitchBend);
-		void SetFreqScale(const float InFreqScale);
+		SIGNALPROCESSING_API void SetCents(const float InCents);
+		SIGNALPROCESSING_API void SetOctave(const float InOctave);
+		SIGNALPROCESSING_API void SetSampleRate(const float InSampleRate);
+		SIGNALPROCESSING_API void SetSemitones(const float InSemiTone);
+		SIGNALPROCESSING_API void SetDetune(const float InDetune);
+		SIGNALPROCESSING_API void SetPitchBend(const float InPitchBend);
+		SIGNALPROCESSING_API void SetFreqScale(const float InFreqScale);
 
 		// Sets the LFO pulse width for square-wave type oscillators
-		void SetPulseWidth(const float InPulseWidth);
+		SIGNALPROCESSING_API void SetPulseWidth(const float InPulseWidth);
 
 		// Returns whether or not this oscillator is playing
 		bool IsPlaying() const { return bIsPlaying; }
@@ -119,10 +119,10 @@ namespace Audio
 
 		// Deprecated: use SetFollowerOsc
 		UE_DEPRECATED(5.1, "SetSlaveOsc is deprecated, please use SetFollowerOsc instead.")
-		void SetSlaveOsc(IOscBase* InSlaveOsc);
+		SIGNALPROCESSING_API void SetSlaveOsc(IOscBase* InSlaveOsc);
 
 		// Sets the input oscillator as the follower of this oscillator
-		void SetFollowerOsc(IOscBase* InFollowerOsc);
+		SIGNALPROCESSING_API void SetFollowerOsc(IOscBase* InFollowerOsc);
 
 		// Return patch destinations for various modulatable parameters
 		FPatchDestination GetModDestFrequency() const { return ModFrequencyDest; }
@@ -247,18 +247,18 @@ namespace Audio
 	}
 
 	// Pitched oscillator
-	class SIGNALPROCESSING_API FOsc : public IOscBase
+	class FOsc : public IOscBase
 	{
 	public:
-		FOsc();
-		virtual ~FOsc();
+		SIGNALPROCESSING_API FOsc();
+		SIGNALPROCESSING_API virtual ~FOsc();
 
 		//~ Begin FOscBase
-		virtual void Start() override;
-		virtual void Stop() override;
-		virtual void Reset() override;
-		virtual void Update() override;
-		virtual float Generate(float* AuxOutput = nullptr) override;
+		SIGNALPROCESSING_API virtual void Start() override;
+		SIGNALPROCESSING_API virtual void Stop() override;
+		SIGNALPROCESSING_API virtual void Reset() override;
+		SIGNALPROCESSING_API virtual void Update() override;
+		SIGNALPROCESSING_API virtual float Generate(float* AuxOutput = nullptr) override;
 		//~ End FOscBase
 
 		// Sets the oscillator type
@@ -270,7 +270,7 @@ namespace Audio
 	protected:
 
 		// Smooth out the saw-tooth discontinuities to improve aliasing
-		float PolySmooth(const float InPhase, const float InPhaseInc);
+		SIGNALPROCESSING_API float PolySmooth(const float InPhase, const float InPhaseInc);
 
 		// Current sign of the square mod, used for triangle wave generation
 		float TriangleSign; 
diff --git a/Engine/Source/Runtime/SignalProcessing/Public/DSP/Phaser.h b/Engine/Source/Runtime/SignalProcessing/Public/DSP/Phaser.h
--- a/Engine/Source/Runtime/SignalProcessing/Public/DSP/Phaser.h
+++ b/Engine/Source/Runtime/SignalProcessing/Public/DSP/Phaser.h
@@ -10,37 +10,37 @@
 namespace Audio
 {
 
-	class SIGNALPROCESSING_API FPhaser
+	class FPhaser
 	{
 	public:
-		FPhaser();
-		~FPhaser();
+		SIGNALPROCESSING_API FPhaser();
+		SIGNALPROCESSING_API ~FPhaser();
 
-		void Init(const float SampleRate, const int32 InNumChannels);
+		SIGNALPROCESSING_API void Init(const float SampleRate, const int32 InNumChannels);
 
 		// Sets the phaser LFO rate
-		void SetFrequency(const float InFreqHz);
+		SIGNALPROCESSING_API void SetFrequency(const float InFreqHz);
 
 		// Sets the wet level of the phaser
-		void SetWetLevel(const float InWetLevel);
+		SIGNALPROCESSING_API void SetWetLevel(const float InWetLevel);
 
 		// Sets the feedback of the phaser
-		void SetFeedback(const float InFeedback);
+		SIGNALPROCESSING_API void SetFeedback(const float InFeedback);
 
 		// Sets the phaser LFO type
-		void SetLFOType(const ELFO::Type LFOType);
+		SIGNALPROCESSING_API void SetLFOType(const ELFO::Type LFOType);
 
 		// Sets whether or not to put the phaser in quadrature mode
-		void SetQuadPhase(const bool bQuadPhase);
+		SIGNALPROCESSING_API void SetQuadPhase(const bool bQuadPhase);
 
 		// Process an audio frame
-		void ProcessAudioFrame(const float* InFrame, float* OutFrame);
+		SIGNALPROCESSING_API void ProcessAudioFrame(const float* InFrame, float* OutFrame);
 
 		// Process an audio buffer.
-		void ProcessAudio(const float* InBuffer, const int32 InNumSamples, float* OutBuffer);
+		SIGNALPROCESSING_API void ProcessAudio(const float* InBuffer, const int32 InNumSamples, float* OutBuffer);
 
 	protected:
-		void ComputeNewCoefficients(const int32 ChannelIndex, const float LFOValue);
+		SIGNALPROCESSING_API void ComputeNewCoefficients(const int32 ChannelIndex, const float LFOValue);
 
 		// First-order all-pass filters in series
 		static const int32 NumApfs = 6;
diff --git a/Engine/Source/Runtime/SignalProcessing/Public/DSP/Reverb.h b/Engine/Source/Runtime/SignalProcessing/Public/DSP/Reverb.h
--- a/Engine/Source/Runtime/SignalProcessing/Public/DSP/Reverb.h
+++ b/Engine/Source/Runtime/SignalProcessing/Public/DSP/Reverb.h
@@ -36,23 +36,23 @@ namespace Audio
 	};
 
 	// Basic implementation of a 4x4 Feedback Delay Network
-	class SIGNALPROCESSING_API FEarlyReflections
+	class FEarlyReflections
 	{
 	public:
-		FEarlyReflections();
-		virtual ~FEarlyReflections();
+		SIGNALPROCESSING_API FEarlyReflections();
+		SIGNALPROCESSING_API virtual ~FEarlyReflections();
 
-		void Init(const int32 InSampleRate);
+		SIGNALPROCESSING_API void Init(const int32 InSampleRate);
 
 		// Sets the reverb settings, applies, and updates
-		void SetSettings(const FEarlyReflectionsSettings& InSettings);
+		SIGNALPROCESSING_API void SetSettings(const FEarlyReflectionsSettings& InSettings);
 
 		// Process the single audio frame
-		void ProcessAudioFrame(const float* InBuffer, const int32 InChannels, float* OutBuffer, const int32 OutChannels);
+		SIGNALPROCESSING_API void ProcessAudioFrame(const float* InBuffer, const int32 InChannels, float* OutBuffer, const int32 OutChannels);
 
 	protected:
-		void ApplySettings();
-		float ProcessDelayLine(const float InSample, FDelayAPF& InAPF, FOnePoleLPF& InLPF);
+		SIGNALPROCESSING_API void ApplySettings();
+		SIGNALPROCESSING_API float ProcessDelayLine(const float InSample, FDelayAPF& InAPF, FOnePoleLPF& InLPF);
 
 		FEarlyReflectionsSettings Settings;
 
@@ -119,30 +119,30 @@ namespace Audio
 		{}
 	};
 
-	class SIGNALPROCESSING_API FPlateReverb
+	class FPlateReverb
 	{
 	public:
-		FPlateReverb();
-		~FPlateReverb();
+		SIGNALPROCESSING_API FPlateReverb();
+		SIGNALPROCESSING_API ~FPlateReverb();
 
 		// Initialize the reverb with the given sample rate
-		void Init(const int32 InSampleRate);
+		SIGNALPROCESSING_API void Init(const int32 InSampleRate);
 
 		// Whether or not to enable late reflections
-		void EnableLateReflections(const bool bInEnableLateReflections);
+		SIGNALPROCESSING_API void EnableLateReflections(const bool bInEnableLateReflections);
 
 		// Whether or not to enable late reflections
-		void EnableEarlyReflections(const bool bInEnableEarlyReflections);
+		SIGNALPROCESSING_API void EnableEarlyReflections(const bool bInEnableEarlyReflections);
 
 		// Sets the reverb settings, applies, and updates
-		void SetSettings(const FPlateReverbSettings& InSettings);
+		SIGNALPROCESSING_API void SetSettings(const FPlateReverbSettings& InSettings);
 
 		// Process the single audio frame
-		void ProcessAudioFrame(const float* InBuffer, const int32 InChannels, float* OutBuffer, const int32 OutChannels);
+		SIGNALPROCESSING_API void ProcessAudioFrame(const float* InBuffer, const int32 InChannels, float* OutBuffer, const int32 OutChannels);
 
 	protected:
 
-		void ApplySettings();
+		SIGNALPROCESSING_API void ApplySettings();
 
 		// Current parameter settings of reverb
 		FPlateReverbSettings Settings;
diff --git a/Engine/Source/Runtime/SignalProcessing/Public/DSP/ReverbFast.h b/Engine/Source/Runtime/SignalProcessing/Public/DSP/ReverbFast.h
--- a/Engine/Source/Runtime/SignalProcessing/Public/DSP/ReverbFast.h
+++ b/Engine/Source/Runtime/SignalProcessing/Public/DSP/ReverbFast.h
@@ -14,7 +14,7 @@
 namespace Audio
 {
 	// Settings for plate reverb
-	struct SIGNALPROCESSING_API FPlateReverbFastSettings
+	struct FPlateReverbFastSettings
 	{
 		// EQuadBehavior describes how reverb is generated when there are 5 or more output channels.
 		enum class EQuadBehavior : uint8
@@ -27,11 +27,11 @@ namespace Audio
 			QuadFlipped
 		};
 
-		FPlateReverbFastSettings();
+		SIGNALPROCESSING_API FPlateReverbFastSettings();
 
-		bool operator==(const FPlateReverbFastSettings& Other) const;
+		SIGNALPROCESSING_API bool operator==(const FPlateReverbFastSettings& Other) const;
 
-		bool operator!=(const FPlateReverbFastSettings& Other) const;
+		SIGNALPROCESSING_API bool operator!=(const FPlateReverbFastSettings& Other) const;
 
 		// EarlyReflectionSettings controls the initial perceived echoes from a sound, modeling the first few
 		// orders of reflections from a sound source to the listener's ears. 
@@ -54,29 +54,29 @@ namespace Audio
 	// The Plate Reverb emulates the interactions between a sound, the listener and the space they share. Early reflections
 	// are modeled using a feedback delay network while late reflections are modeled using a plate reverb. This 
 	// class aims to support a flexible and pleasant sounding reverb balanced with computational efficiency. 
-	class SIGNALPROCESSING_API FPlateReverbFast {
+	class FPlateReverbFast {
 		public:
-			static const FPlateReverbFastSettings DefaultSettings;
+			static SIGNALPROCESSING_API const FPlateReverbFastSettings DefaultSettings;
 			
 			// InMaxInternalBufferSamples sets the maximum number of samples used in internal buffers.
-			FPlateReverbFast(float InSampleRate, int32 InMaxInternalBufferSamples = 512, const FPlateReverbFastSettings& InSettings=DefaultSettings);
+			SIGNALPROCESSING_API FPlateReverbFast(float InSampleRate, int32 InMaxInternalBufferSamples = 512, const FPlateReverbFastSettings& InSettings=DefaultSettings);
 
-			~FPlateReverbFast();
+			SIGNALPROCESSING_API ~FPlateReverbFast();
 
 			// Copies, clamps and applies settings.
-			void SetSettings(const FPlateReverbFastSettings& InSettings);
+			SIGNALPROCESSING_API void SetSettings(const FPlateReverbFastSettings& InSettings);
 
-			const FPlateReverbFastSettings& GetSettings() const;
+			SIGNALPROCESSING_API const FPlateReverbFastSettings& GetSettings() const;
 
 			// Creates reverberated audio in OutSamples based upon InSamples
 			// InNumChannels can be 1 or 2 channels.
 			// OutSamples must be greater or equal to 2.
-			void ProcessAudio(const FAlignedFloatBuffer& InSamples, const int32 InNumChannels, FAlignedFloatBuffer& OutSamples, const int32 OutNumChannels);
+			SIGNALPROCESSING_API void ProcessAudio(const FAlignedFloatBuffer& InSamples, const int32 InNumChannels, FAlignedFloatBuffer& OutSamples, const int32 OutNumChannels);
 
-			void FlushAudio();
+			SIGNALPROCESSING_API void FlushAudio();
 
 			// Clamp individual settings to values supported by this class.
-			static void ClampSettings(FPlateReverbFastSettings& InOutSettings);
+			static SIGNALPROCESSING_API void ClampSettings(FPlateReverbFastSettings& InOutSettings);
 
 		private:
 			// Copy reverberated samples to interleaved output samples. Map channels according to internal settings.
diff --git a/Engine/Source/Runtime/SignalProcessing/Public/DSP/RingModulation.h b/Engine/Source/Runtime/SignalProcessing/Public/DSP/RingModulation.h
--- a/Engine/Source/Runtime/SignalProcessing/Public/DSP/RingModulation.h
+++ b/Engine/Source/Runtime/SignalProcessing/Public/DSP/RingModulation.h
@@ -10,41 +10,41 @@ namespace Audio
 {
 	// Ring modulation effect
 	// https://en.wikipedia.org/wiki/Ring_modulation
-	class SIGNALPROCESSING_API FRingModulation
+	class FRingModulation
 	{
 	public:
 		// Constructor
-		FRingModulation();
+		SIGNALPROCESSING_API FRingModulation();
 
 		// Destructor
-		~FRingModulation();
+		SIGNALPROCESSING_API ~FRingModulation();
 
 		// Initialize the equalizer
-		void Init(const float InSampleRate, const int32 InNumChannels);
+		SIGNALPROCESSING_API void Init(const float InSampleRate, const int32 InNumChannels);
 
 		// The type of modulation
-		void SetModulatorWaveType(const EOsc::Type InType);
+		SIGNALPROCESSING_API void SetModulatorWaveType(const EOsc::Type InType);
 
 		// Set the ring modulation frequency
-		void SetModulationFrequency(const float InModulationFrequency);
+		SIGNALPROCESSING_API void SetModulationFrequency(const float InModulationFrequency);
 
 		// Set the ring modulation depth
-		void SetModulationDepth(const float InModulationDepth);
+		SIGNALPROCESSING_API void SetModulationDepth(const float InModulationDepth);
 
 		// Sets that the modulation buffer is external
-		void SetExternalPatchSource(Audio::FPatchOutputStrongPtr InPatch);
+		SIGNALPROCESSING_API void SetExternalPatchSource(Audio::FPatchOutputStrongPtr InPatch);
 
 		// Set the dry level of the ring modulation
 		void SetDryLevel(const float InDryLevel) { DryLevel = InDryLevel; }
 
 		// Set the wet level of the ring modulation
-		void SetWetLevel(const float InWetLevel);
+		SIGNALPROCESSING_API void SetWetLevel(const float InWetLevel);
 
 		// Process audio buffer
-		void ProcessAudio(const float* InBuffer, const int32 InNumSamples, float* OutBuffer);
+		SIGNALPROCESSING_API void ProcessAudio(const float* InBuffer, const int32 InNumSamples, float* OutBuffer);
 
 	private:
-		void UpdateScale();
+		SIGNALPROCESSING_API void UpdateScale();
 
 		Audio::FPatchOutputStrongPtr Patch;
 		Audio::FOsc Osc;
diff --git a/Engine/Source/Runtime/SignalProcessing/Public/DSP/SampleBufferReader.h b/Engine/Source/Runtime/SignalProcessing/Public/DSP/SampleBufferReader.h
--- a/Engine/Source/Runtime/SignalProcessing/Public/DSP/SampleBufferReader.h
+++ b/Engine/Source/Runtime/SignalProcessing/Public/DSP/SampleBufferReader.h
@@ -18,28 +18,28 @@ namespace Audio
 		};
 	}
 
-	class SIGNALPROCESSING_API FSampleBufferReader
+	class FSampleBufferReader
 	{
 	public:
-		FSampleBufferReader();
-		~FSampleBufferReader();
+		SIGNALPROCESSING_API FSampleBufferReader();
+		SIGNALPROCESSING_API ~FSampleBufferReader();
 
-		void Init(const int32 InSampleRate);
+		SIGNALPROCESSING_API void Init(const int32 InSampleRate);
 
 		// This must be a completely loaded buffer. This buffer reader doesn't OWN the buffer memory.
-		void SetBuffer(const int16* InBufferPtr, const int32 InNumBufferSamples, const int32 InNumChannels, const int32 InBufferSampleRate);
+		SIGNALPROCESSING_API void SetBuffer(const int16* InBufferPtr, const int32 InNumBufferSamples, const int32 InNumChannels, const int32 InBufferSampleRate);
 
 		// Seeks the buffer the given time in seconds. Returns true if succeeded.
-		void SeekTime(const float InTimeSec, const ESeekType::Type InSeekType = ESeekType::FromBeginning, const bool bWrap = true);
+		SIGNALPROCESSING_API void SeekTime(const float InTimeSec, const ESeekType::Type InSeekType = ESeekType::FromBeginning, const bool bWrap = true);
 
 		// Sets the pitch of the buffer reader. Can be negative. Will linearly interpolate over the given time value.
-		void SetPitch(const float InPitch, const float InterpolationTimeSec = 0.0f);
+		SIGNALPROCESSING_API void SetPitch(const float InPitch, const float InterpolationTimeSec = 0.0f);
 
 		// Puts the wave reader into scrub mode
-		void SetScrubMode(const bool bInIsScrubMode);
+		SIGNALPROCESSING_API void SetScrubMode(const bool bInIsScrubMode);
 		
 		// Sets the scrub width. The sound will loop between the scrub width region and the current frame
-		void SetScrubTimeWidth(const float InScrubTimeWidthSec);
+		SIGNALPROCESSING_API void SetScrubTimeWidth(const float InScrubTimeWidthSec);
 
 		// Returns the number of channels of this buffer.
 		int32 GetNumChannels() const { return BufferNumChannels; }
@@ -51,20 +51,20 @@ namespace Audio
 		float GetPlaybackProgress() const { return PlaybackProgress; }
 
 		// Generates the next block of audio. Returns true if it's no longer playing (reached end of the buffer and not set to wrap)
-		bool Generate(float* OutAudioBuffer, const int32 NumFrames, const int32 OutChannels, const bool bWrap = false);
+		SIGNALPROCESSING_API bool Generate(float* OutAudioBuffer, const int32 NumFrames, const int32 OutChannels, const bool bWrap = false);
 
 		// Whether or not the buffer reader has a buffer
 		bool HasBuffer() const { return BufferPtr != nullptr; }
 
 		// Clears current buffer and resets state
-		void ClearBuffer();
+		SIGNALPROCESSING_API void ClearBuffer();
 
 	protected:
 
-		float GetSampleValueForChannel(const int32 Channel);
-		void UpdateScrubMinAndMax();
-		void UpdateSeekFrame();
-		float GetSampleValue(const int16* InBuffer, const int32 SampleIndex);
+		SIGNALPROCESSING_API float GetSampleValueForChannel(const int32 Channel);
+		SIGNALPROCESSING_API void UpdateScrubMinAndMax();
+		SIGNALPROCESSING_API void UpdateSeekFrame();
+		SIGNALPROCESSING_API float GetSampleValue(const int16* InBuffer, const int32 SampleIndex);
 
 		const int16* BufferPtr;
 		int32 BufferNumSamples;
diff --git a/Engine/Source/Runtime/SignalProcessing/Public/DSP/SampleRateConverter.h b/Engine/Source/Runtime/SignalProcessing/Public/DSP/SampleRateConverter.h
--- a/Engine/Source/Runtime/SignalProcessing/Public/DSP/SampleRateConverter.h
+++ b/Engine/Source/Runtime/SignalProcessing/Public/DSP/SampleRateConverter.h
@@ -6,7 +6,7 @@
 
 namespace Audio
 {
-	class SIGNALPROCESSING_API ISampleRateConverter
+	class ISampleRateConverter
 	{
 	public:
 		virtual ~ISampleRateConverter() {}
@@ -39,7 +39,7 @@ namespace Audio
 		virtual int32 ProcessChunk(const float* BufferChunk, const int32 NumInputSamples, const int32 RequestedFrames, TArray<float>& OutBuffer) = 0;
 
 		// Creates a new instance
-		static ISampleRateConverter* CreateSampleRateConverter();
+		static SIGNALPROCESSING_API ISampleRateConverter* CreateSampleRateConverter();
 	};
 
 }
diff --git a/Engine/Source/Runtime/SignalProcessing/Public/DSP/SinOsc.h b/Engine/Source/Runtime/SignalProcessing/Public/DSP/SinOsc.h
--- a/Engine/Source/Runtime/SignalProcessing/Public/DSP/SinOsc.h
+++ b/Engine/Source/Runtime/SignalProcessing/Public/DSP/SinOsc.h
@@ -13,35 +13,35 @@ namespace Audio
 	* Setting frequency uses current filter state to compute initial conditions to avoid pops when changing frequency.
 	* Extremely cheap to run but expensive to set new frequencies. Good for test tones.
 	*/
-	class SIGNALPROCESSING_API FSineOsc
+	class FSineOsc
 	{
 	public:
 		/** Constructor */
-		FSineOsc();
+		SIGNALPROCESSING_API FSineOsc();
 
 		/** Non-default constructor */
-		FSineOsc(const int32 InSampleRate, const float InFrequencyHz, const float Scale = 1.0f, const float Add = 0.0f);
+		SIGNALPROCESSING_API FSineOsc(const int32 InSampleRate, const float InFrequencyHz, const float Scale = 1.0f, const float Add = 0.0f);
 
 		/** Virtual destructor */
-		virtual ~FSineOsc();
+		SIGNALPROCESSING_API virtual ~FSineOsc();
 
 		/** Initialize the oscillator with a sample rate and new frequency. Must be called before playing oscillator. */
-		void Init(const int32 InSampleRate, const float InFrequencyHz, const float Scale = 1.0f, const float Add = 0.0f);
+		SIGNALPROCESSING_API void Init(const int32 InSampleRate, const float InFrequencyHz, const float Scale = 1.0f, const float Add = 0.0f);
 
 		/** Sets the scale of the oscillator. */
-		void SetScale(const float InScale);
+		SIGNALPROCESSING_API void SetScale(const float InScale);
 
 		/** Sets the scale of the oscillator. */
-		void SetAdd(const float InAdd);
+		SIGNALPROCESSING_API void SetAdd(const float InAdd);
 
 		/** Sets the frequency of the oscillator in Hz (based on sample rate). Performs initial condition calculation to avoid pops. */
-		void SetFrequency(const float InFrequencyHz);
+		SIGNALPROCESSING_API void SetFrequency(const float InFrequencyHz);
 
 		/** Returns the current frequency. */
-		float GetFrequency() const;
+		SIGNALPROCESSING_API float GetFrequency() const;
 
 		/** Generates the next sample of the oscillator. */
-		float ProcessAudio();
+		SIGNALPROCESSING_API float ProcessAudio();
 
 	protected:
 		int32 SampleRate;
diff --git a/Engine/Source/Runtime/SignalProcessing/Public/DSP/SineWaveTableOsc.h b/Engine/Source/Runtime/SignalProcessing/Public/DSP/SineWaveTableOsc.h
--- a/Engine/Source/Runtime/SignalProcessing/Public/DSP/SineWaveTableOsc.h
+++ b/Engine/Source/Runtime/SignalProcessing/Public/DSP/SineWaveTableOsc.h
@@ -10,41 +10,41 @@ namespace Audio
 
 
 	// A sinusoidal wave table oscillator class
-	class SIGNALPROCESSING_API FSineWaveTableOsc
+	class FSineWaveTableOsc
 	{
 	public:
 		// Constructor
-		FSineWaveTableOsc();
+		SIGNALPROCESSING_API FSineWaveTableOsc();
 
 		// Virtual Destructor
-		virtual ~FSineWaveTableOsc();
+		SIGNALPROCESSING_API virtual ~FSineWaveTableOsc();
 
 		// Initialize the wave table oscillator
-		void Init(const float InSampleRate, const float InFrequencyHz, const float InPhase);
+		SIGNALPROCESSING_API void Init(const float InSampleRate, const float InFrequencyHz, const float InPhase);
 
 		// Sets the sample rate of the oscillator.
-		void SetSampleRate(const float InSampleRate);
+		SIGNALPROCESSING_API void SetSampleRate(const float InSampleRate);
 
 		// Resets the wave table read indices.
-		void Reset();
+		SIGNALPROCESSING_API void Reset();
 
 		// Sets the frequency of the wave table oscillator.
-		void SetFrequencyHz(const float InFrequencyHz);
+		SIGNALPROCESSING_API void SetFrequencyHz(const float InFrequencyHz);
 
 		// Returns the frequency of the wave table oscillator.
 		float GetFrequencyHz() const { return FrequencyHz; }
 		
 		// Sets the phase of the wave table oscillator.
-		void SetPhase(const float InPhase);
+		SIGNALPROCESSING_API void SetPhase(const float InPhase);
 
 		// Processes the wave table and fills a buffer
-		void Generate(float* OutBuffer, const int32 NumSamples);
+		SIGNALPROCESSING_API void Generate(float* OutBuffer, const int32 NumSamples);
 
 		// The static sinusoidal wave table 
-		static const TArray<float>& GetWaveTable();
+		static SIGNALPROCESSING_API const TArray<float>& GetWaveTable();
 
 	protected:
-		void UpdatePhaseIncrement();
+		SIGNALPROCESSING_API void UpdatePhaseIncrement();
 
 		// The wave table buffer
 		const TArray<float>& WaveTableBuffer = GetWaveTable();
diff --git a/Engine/Source/Runtime/SignalProcessing/Public/DSP/SpectrumAnalyzer.h b/Engine/Source/Runtime/SignalProcessing/Public/DSP/SpectrumAnalyzer.h
--- a/Engine/Source/Runtime/SignalProcessing/Public/DSP/SpectrumAnalyzer.h
+++ b/Engine/Source/Runtime/SignalProcessing/Public/DSP/SpectrumAnalyzer.h
@@ -14,7 +14,7 @@ namespace Audio
 	class IFFTAlgorithm;
 	class FSpectrumAnalyzer;
 
-	struct SIGNALPROCESSING_API FSpectrumAnalyzerSettings
+	struct FSpectrumAnalyzerSettings
 	{
 		// Actual FFT size used. For FSpectrumAnalyzer, we never zero pad the input buffer.
 		enum class EFFTSize : uint16
@@ -49,7 +49,7 @@ namespace Audio
 
 
 	/** Settings for band extractor. */
-	struct SIGNALPROCESSING_API FSpectrumBandExtractorSettings
+	struct FSpectrumBandExtractorSettings
 	{
 		/** Metric for output band values. */
 		enum class EMetric : uint8
@@ -94,7 +94,7 @@ namespace Audio
 	};
 
 	/** Settings describing the spectrum used for in the band extractor. */
-	struct SIGNALPROCESSING_API FSpectrumBandExtractorSpectrumSettings
+	struct FSpectrumBandExtractorSpectrumSettings
 	{
 		/** Sample rate of audio */
 		float SampleRate; 
@@ -139,7 +139,7 @@ namespace Audio
 	 *  By maintaining band information across multiple calls, some intermediate 
 	 *  values can be cached to speed up the operation.
 	 */
-	class SIGNALPROCESSING_API ISpectrumBandExtractor
+	class ISpectrumBandExtractor
 	{
 		public:
 
@@ -200,7 +200,7 @@ namespace Audio
 			virtual void ExtractBands(const FAlignedFloatBuffer& InComplexBuffer, double InTimestamp, TArray<float>& OutValues) = 0;
 
 			/** Creates a ISpectrumBandExtractor. */
-			static TUniquePtr<ISpectrumBandExtractor> CreateSpectrumBandExtractor(const FSpectrumBandExtractorSettings& InSettings);
+			static SIGNALPROCESSING_API TUniquePtr<ISpectrumBandExtractor> CreateSpectrumBandExtractor(const FSpectrumBandExtractorSettings& InSettings);
 	};
 
 	/**
@@ -287,7 +287,7 @@ namespace Audio
 	 * Typical usage is to either call PushAudio() and then PerformAnalysisIfPossible immediately afterwards,
 	 * or have a seperate thread call PerformAnalysisIfPossible().
 	 */
-	class SIGNALPROCESSING_API FSpectrumAnalyzer 
+	class FSpectrumAnalyzer 
 	{
 
 	public:
@@ -301,62 +301,62 @@ namespace Audio
 		};
 
 		// If an instance is created using the default constructor, Init() must be called before it is used.
-		FSpectrumAnalyzer();
+		SIGNALPROCESSING_API FSpectrumAnalyzer();
 
 		// If an instance is created using either of these constructors, Init() is not neccessary.
-		FSpectrumAnalyzer(float InSampleRate);
-		FSpectrumAnalyzer(const FSpectrumAnalyzerSettings& InSettings, float InSampleRate);
+		SIGNALPROCESSING_API FSpectrumAnalyzer(float InSampleRate);
+		SIGNALPROCESSING_API FSpectrumAnalyzer(const FSpectrumAnalyzerSettings& InSettings, float InSampleRate);
 
 		virtual ~FSpectrumAnalyzer() = default;
 
 		// Initialize sample rate of analyzer if not known at time of construction
-		void Init(float InSampleRate);
-		void Init(const FSpectrumAnalyzerSettings& InSettings, float InSampleRate);
+		SIGNALPROCESSING_API void Init(float InSampleRate);
+		SIGNALPROCESSING_API void Init(const FSpectrumAnalyzerSettings& InSettings, float InSampleRate);
 
 		// Update the settings used by this Spectrum Analyzer. Safe to call on any thread, but should not be called every tick.
-		void SetSettings(const FSpectrumAnalyzerSettings& InSettings);
+		SIGNALPROCESSING_API void SetSettings(const FSpectrumAnalyzerSettings& InSettings);
 
 		// Get the current settings used by this Spectrum Analyzer.
-		void GetSettings(FSpectrumAnalyzerSettings& OutSettings);
+		SIGNALPROCESSING_API void GetSettings(FSpectrumAnalyzerSettings& OutSettings);
 
 		// Samples magnitude (linearly) for a given frequency, in Hz.
-		float GetMagnitudeForFrequency(float InFrequency, EPeakInterpolationMethod InMethod = EPeakInterpolationMethod::Linear);
-		float GetNormalizedMagnitudeForFrequency(float InFrequency, EPeakInterpolationMethod InMethod = EPeakInterpolationMethod::Linear);
+		SIGNALPROCESSING_API float GetMagnitudeForFrequency(float InFrequency, EPeakInterpolationMethod InMethod = EPeakInterpolationMethod::Linear);
+		SIGNALPROCESSING_API float GetNormalizedMagnitudeForFrequency(float InFrequency, EPeakInterpolationMethod InMethod = EPeakInterpolationMethod::Linear);
 
 		// Samples phase for a given frequency, in Hz.
-		float GetPhaseForFrequency(float InFrequency, EPeakInterpolationMethod InMethod = EPeakInterpolationMethod::Linear);
+		SIGNALPROCESSING_API float GetPhaseForFrequency(float InFrequency, EPeakInterpolationMethod InMethod = EPeakInterpolationMethod::Linear);
 
 		// Return array of bands using spectrum band extractor.
-		void GetBands(ISpectrumBandExtractor& InExtractor, TArray<float>& OutValues);
+		SIGNALPROCESSING_API void GetBands(ISpectrumBandExtractor& InExtractor, TArray<float>& OutValues);
 
 		// You can call this function to ensure that you're sampling the same window of frequency data,
 		// Then call UnlockOutputBuffer when you're done.
 		// Otherwise, GetMagnitudeForFrequency and GetPhaseForFrequency will always use the latest window
 		// of frequency data.
-		void LockOutputBuffer();
-		void UnlockOutputBuffer();
+		SIGNALPROCESSING_API void LockOutputBuffer();
+		SIGNALPROCESSING_API void UnlockOutputBuffer();
 		
 		// Push audio to queue. Returns false if the queue is already full.
-		bool PushAudio(const TSampleBuffer<float>& InBuffer);
-		bool PushAudio(const float* InBuffer, int32 NumSamples);
+		SIGNALPROCESSING_API bool PushAudio(const TSampleBuffer<float>& InBuffer);
+		SIGNALPROCESSING_API bool PushAudio(const float* InBuffer, int32 NumSamples);
 
 		// Thread safe call to perform actual FFT. Returns true if it performed the FFT, false otherwise.
 		// If bUseLatestAudio is set to true, this function will flush the entire input buffer, potentially losing data.
 		// Otherwise it will only consume enough samples necessary to perform a single FFT.
-		bool PerformAnalysisIfPossible(bool bUseLatestAudio = false);
+		SIGNALPROCESSING_API bool PerformAnalysisIfPossible(bool bUseLatestAudio = false);
 
 		// Returns false if this instance of FSpectrumAnalyzer was constructed with the default constructor 
 		// and Init() has not been called yet.
-		bool IsInitialized();
+		SIGNALPROCESSING_API bool IsInitialized();
 	
 	private:
 
 
 		// Called on analysis thread.
-		void ResetSettings();
+		SIGNALPROCESSING_API void ResetSettings();
 
 		// Called in GetMagnitudeForFrequency and GetPhaseForFrequency.
-		void PerformInterpolation(const FAlignedFloatBuffer& InComplexBuffer, EPeakInterpolationMethod InMethod, const float InFreq, float& OutReal, float& OutImag);
+		SIGNALPROCESSING_API void PerformInterpolation(const FAlignedFloatBuffer& InComplexBuffer, EPeakInterpolationMethod InMethod, const float InFreq, float& OutReal, float& OutImag);
 
 		// Cached current settings. Only actually used in ResetSettings().
 		FSpectrumAnalyzerSettings CurrentSettings;
@@ -385,7 +385,7 @@ namespace Audio
 		TUniquePtr<IFFTAlgorithm> FFT;
 	};
 
-	class SIGNALPROCESSING_API FSpectrumAnalyzerScopeLock
+	class FSpectrumAnalyzerScopeLock
 	{
 	public:
 		FSpectrumAnalyzerScopeLock(FSpectrumAnalyzer* InAnalyzer)
@@ -404,7 +404,7 @@ namespace Audio
 	};
 
 	// SpectrumAnalyzer for computing spectrum in async task. 
-	class SIGNALPROCESSING_API FAsyncSpectrumAnalyzer 
+	class FAsyncSpectrumAnalyzer 
 	{
 
 		FAsyncSpectrumAnalyzer(const FSpectrumAnalyzer&) = delete;
@@ -415,58 +415,58 @@ namespace Audio
 
 
 	public:
-		FAsyncSpectrumAnalyzer();
+		SIGNALPROCESSING_API FAsyncSpectrumAnalyzer();
 		// If an instance is created using either of these constructors, Init() is not neccessary.
-		FAsyncSpectrumAnalyzer(float InSampleRate);
-		FAsyncSpectrumAnalyzer(const FSpectrumAnalyzerSettings& InSettings, float InSampleRate);
+		SIGNALPROCESSING_API FAsyncSpectrumAnalyzer(float InSampleRate);
+		SIGNALPROCESSING_API FAsyncSpectrumAnalyzer(const FSpectrumAnalyzerSettings& InSettings, float InSampleRate);
 
-		virtual ~FAsyncSpectrumAnalyzer();
+		SIGNALPROCESSING_API virtual ~FAsyncSpectrumAnalyzer();
 
 		// Initialize sample rate of analyzer if not known at time of construction
-		void Init(float InSampleRate);
-		void Init(const FSpectrumAnalyzerSettings& InSettings, float InSampleRate);
+		SIGNALPROCESSING_API void Init(float InSampleRate);
+		SIGNALPROCESSING_API void Init(const FSpectrumAnalyzerSettings& InSettings, float InSampleRate);
 
 		// Returns false if this instance of FSpectrumAnalyzer was constructed with the default constructor 
 		// and Init() has not been called yet.
-		bool IsInitialized();
+		SIGNALPROCESSING_API bool IsInitialized();
 
 		// Update the settings used by this Spectrum Analyzer. Safe to call on any thread, but should not be called every tick.
-		void SetSettings(const FSpectrumAnalyzerSettings& InSettings);
+		SIGNALPROCESSING_API void SetSettings(const FSpectrumAnalyzerSettings& InSettings);
 
 		// Get the current settings used by this Spectrum Analyzer.
-		void GetSettings(FSpectrumAnalyzerSettings& OutSettings);
+		SIGNALPROCESSING_API void GetSettings(FSpectrumAnalyzerSettings& OutSettings);
 
 		// Samples magnitude (linearly) for a given frequency, in Hz.
-		float GetMagnitudeForFrequency(float InFrequency, FSpectrumAnalyzer::EPeakInterpolationMethod InMethod = FSpectrumAnalyzer::EPeakInterpolationMethod::Linear);
-		float GetNormalizedMagnitudeForFrequency(float InFrequency, FSpectrumAnalyzer::EPeakInterpolationMethod InMethod = FSpectrumAnalyzer::EPeakInterpolationMethod::Linear);
+		SIGNALPROCESSING_API float GetMagnitudeForFrequency(float InFrequency, FSpectrumAnalyzer::EPeakInterpolationMethod InMethod = FSpectrumAnalyzer::EPeakInterpolationMethod::Linear);
+		SIGNALPROCESSING_API float GetNormalizedMagnitudeForFrequency(float InFrequency, FSpectrumAnalyzer::EPeakInterpolationMethod InMethod = FSpectrumAnalyzer::EPeakInterpolationMethod::Linear);
 
 		// Samples phase for a given frequency, in Hz.
-		float GetPhaseForFrequency(float InFrequency, FSpectrumAnalyzer::EPeakInterpolationMethod InMethod = FSpectrumAnalyzer::EPeakInterpolationMethod::Linear);
+		SIGNALPROCESSING_API float GetPhaseForFrequency(float InFrequency, FSpectrumAnalyzer::EPeakInterpolationMethod InMethod = FSpectrumAnalyzer::EPeakInterpolationMethod::Linear);
 
 		// Return array of bands using spectrum band extractor.
-		void GetBands(ISpectrumBandExtractor& InExtractor, TArray<float>& OutValues);
+		SIGNALPROCESSING_API void GetBands(ISpectrumBandExtractor& InExtractor, TArray<float>& OutValues);
 
 		// You can call this function to ensure that you're sampling the same window of frequency data,
 		// Then call UnlockOutputBuffer when you're done.
 		// Otherwise, GetMagnitudeForFrequency and GetPhaseForFrequency will always use the latest window
 		// of frequency data.
-		void LockOutputBuffer();
-		void UnlockOutputBuffer();
+		SIGNALPROCESSING_API void LockOutputBuffer();
+		SIGNALPROCESSING_API void UnlockOutputBuffer();
 		
 		// Push audio to queue. Returns false if the queue is already full.
-		bool PushAudio(const TSampleBuffer<float>& InBuffer);
-		bool PushAudio(const float* InBuffer, int32 NumSamples);
+		SIGNALPROCESSING_API bool PushAudio(const TSampleBuffer<float>& InBuffer);
+		SIGNALPROCESSING_API bool PushAudio(const float* InBuffer, int32 NumSamples);
 
 		// Thread safe call to perform actual FFT. Returns true if it performed the FFT, false otherwise.
 		// If bUseLatestAudio is set to true, this function will flush the entire input buffer, potentially losing data.
 		// Otherwise it will only consume enough samples necessary to perform a single FFT.
-		bool PerformAnalysisIfPossible(bool bUseLatestAudio = false);
+		SIGNALPROCESSING_API bool PerformAnalysisIfPossible(bool bUseLatestAudio = false);
 
 
 		// Thread safe call to perform actual FFT. Returns true if it performed the FFT, false otherwise.
 		// If bUseLatestAudio is set to true, this function will flush the entire input buffer, potentially losing data.
 		// Otherwise it will only consume enough samples necessary to perform a single FFT.
-		bool PerformAsyncAnalysisIfPossible(bool bUseLatestAudio = false);
+		SIGNALPROCESSING_API bool PerformAsyncAnalysisIfPossible(bool bUseLatestAudio = false);
 	
 
 	private:
@@ -477,7 +477,7 @@ namespace Audio
 		TUniquePtr<FSpectrumAnalyzerTask> AsyncAnalysisTask;
 	};
 
-	class SIGNALPROCESSING_API FAsyncSpectrumAnalyzerScopeLock
+	class FAsyncSpectrumAnalyzerScopeLock
 	{
 	public:
 		FAsyncSpectrumAnalyzerScopeLock(FAsyncSpectrumAnalyzer* InAnalyzer)
diff --git a/Engine/Source/Runtime/SignalProcessing/Public/DSP/TapDelayPitchShifter.h b/Engine/Source/Runtime/SignalProcessing/Public/DSP/TapDelayPitchShifter.h
--- a/Engine/Source/Runtime/SignalProcessing/Public/DSP/TapDelayPitchShifter.h
+++ b/Engine/Source/Runtime/SignalProcessing/Public/DSP/TapDelayPitchShifter.h
@@ -9,7 +9,7 @@
 namespace Audio
 {
 	// Class which reads doppler-shifted tap from a delay buffer 
-	class SIGNALPROCESSING_API FTapDelayPitchShifter
+	class FTapDelayPitchShifter
 	{
 	public:
 		static constexpr float MinDelayLength = 10.0f;
@@ -17,31 +17,31 @@ namespace Audio
 		static constexpr float MaxAbsPitchShiftInOctaves = 6.0f;
 		
 		// Constructor
-		FTapDelayPitchShifter();
+		SIGNALPROCESSING_API FTapDelayPitchShifter();
 
 		// Virtual Destructor
-		virtual ~FTapDelayPitchShifter();
+		SIGNALPROCESSING_API virtual ~FTapDelayPitchShifter();
 
 		// Initialization of the internal delay with given sample rate, pitch shift in semitones, and delay length in milliseconds
-		void Init(const float InSampleRate, const float InPitchShift, const float InDelayLength);
+		SIGNALPROCESSING_API void Init(const float InSampleRate, const float InPitchShift, const float InDelayLength);
 
 		// Sets the internal delay line length.
-		void SetDelayLength(const float InDelayLength);
+		SIGNALPROCESSING_API void SetDelayLength(const float InDelayLength);
 
 		// Sets the pitch shift in semitones
-		void SetPitchShift(const float InPitchScaleSemitones);
+		SIGNALPROCESSING_API void SetPitchShift(const float InPitchScaleSemitones);
 
 		// Sets the pitch shift in terms of a sample rate ratio
-		void SetPitchShiftRatio(const float InPitchShiftRatio);
+		SIGNALPROCESSING_API void SetPitchShiftRatio(const float InPitchShiftRatio);
 		
 		// Reads the next sample tap from the input delay buffer
-		float ReadDopplerShiftedTapFromDelay(const Audio::FDelay& InDelayBuffer, const float ReadOffsetMilliseconds = 0.0f);
+		SIGNALPROCESSING_API float ReadDopplerShiftedTapFromDelay(const Audio::FDelay& InDelayBuffer, const float ReadOffsetMilliseconds = 0.0f);
 
 		// Process whole block of audio reading to and writing from delay line from input buffer
-		void ProcessAudio(Audio::FDelay& InDelayBuffer, const float* InAudioBuffer, const int32 InNumFrames, float* OutAudioBuffer);
+		SIGNALPROCESSING_API void ProcessAudio(Audio::FDelay& InDelayBuffer, const float* InAudioBuffer, const int32 InNumFrames, float* OutAudioBuffer);
 		
 	private:
-		void UpdatePhasorPhaseIncrement();
+		SIGNALPROCESSING_API void UpdatePhasorPhaseIncrement();
 		
 		float SampleRate = 0.0f;
 		Audio::FExponentialEase CurrentDelayLength;
diff --git a/Engine/Source/Runtime/SignalProcessing/Public/DSP/VariablePoleFilter.h b/Engine/Source/Runtime/SignalProcessing/Public/DSP/VariablePoleFilter.h
--- a/Engine/Source/Runtime/SignalProcessing/Public/DSP/VariablePoleFilter.h
+++ b/Engine/Source/Runtime/SignalProcessing/Public/DSP/VariablePoleFilter.h
@@ -15,7 +15,7 @@ namespace Audio
 		EightPole,
 	};
 
-	class SIGNALPROCESSING_API FVariablePoleFilter
+	class FVariablePoleFilter
 	{
 	public:
 		int32 NumFilters = 0;
@@ -25,7 +25,7 @@ namespace Audio
 		EBiquadFilter::Type Type = EBiquadFilter::Lowpass;
 		FVariablePoleFilter() {};
 
-		void Init(EFilterOrder InOrder
+		SIGNALPROCESSING_API void Init(EFilterOrder InOrder
 				  , const float InSampleRate
 				  , const int32 InNumChannels
 				  , const float InFrequency
@@ -33,16 +33,16 @@ namespace Audio
 				  , const float InBandwidth = 2.f
 				  , const float InGain = 0.f);
 
-		void SetParams(const EBiquadFilter::Type InFilterType, const float InCutoffFrequency, const float InBandwidth = 2.f, const float InGain = 0.f);
+		SIGNALPROCESSING_API void SetParams(const EBiquadFilter::Type InFilterType, const float InCutoffFrequency, const float InBandwidth = 2.f, const float InGain = 0.f);
 
-		void ProcessAudioFrame(const float* InFrame, float* OutFrame);
-		void ProcessAudioBuffer(const float* InFrame, float* OutFrame, const int32 NumSamples);
+		SIGNALPROCESSING_API void ProcessAudioFrame(const float* InFrame, float* OutFrame);
+		SIGNALPROCESSING_API void ProcessAudioBuffer(const float* InFrame, float* OutFrame, const int32 NumSamples);
 
-		void Reset();
+		SIGNALPROCESSING_API void Reset();
 
 		const TArray<FBiquadFilter>& GetFilters() const { return Filters; }
 
 	private:
 		TArray<FBiquadFilter> Filters;
 	};
-}
\ No newline at end of file
+}
diff --git a/Engine/Source/Runtime/SignalProcessing/Public/DSP/VectorLinearResampler.h b/Engine/Source/Runtime/SignalProcessing/Public/DSP/VectorLinearResampler.h
--- a/Engine/Source/Runtime/SignalProcessing/Public/DSP/VectorLinearResampler.h
+++ b/Engine/Source/Runtime/SignalProcessing/Public/DSP/VectorLinearResampler.h
@@ -89,7 +89,7 @@ namespace Audio
 			}
 
 	*/
-	struct SIGNALPROCESSING_API FVectorLinearResampler
+	struct FVectorLinearResampler
 	{
 		uint32 CurrentFrameFraction;
 
@@ -109,7 +109,7 @@ namespace Audio
 		//
 		// stereo is deinterleaved, with the right channels being specified by Frames+StrideFloats.
 		//
-		uint32 ResampleMono(uint32 OutputFramesNeeded, uint32 FixedPointSampleRate, float const* SourceFrames, float* OutputFrames);
-		uint32 ResampleStereo(uint32 OutputFramesNeeded, uint32 FixedPointSampleRate, float const* SourceFrames, uint32 SourceFramesStrideFloats, float* OutputFrames, uint32 OutputFramesStrideFloats);
+		SIGNALPROCESSING_API uint32 ResampleMono(uint32 OutputFramesNeeded, uint32 FixedPointSampleRate, float const* SourceFrames, float* OutputFrames);
+		SIGNALPROCESSING_API uint32 ResampleStereo(uint32 OutputFramesNeeded, uint32 FixedPointSampleRate, float const* SourceFrames, uint32 SourceFramesStrideFloats, float* OutputFrames, uint32 OutputFramesStrideFloats);
 	};
-}
\ No newline at end of file
+}
diff --git a/Engine/Source/Runtime/SignalProcessing/Public/DSP/VoiceProcessing.h b/Engine/Source/Runtime/SignalProcessing/Public/DSP/VoiceProcessing.h
--- a/Engine/Source/Runtime/SignalProcessing/Public/DSP/VoiceProcessing.h
+++ b/Engine/Source/Runtime/SignalProcessing/Public/DSP/VoiceProcessing.h
@@ -14,16 +14,16 @@ namespace Audio
 	 * This is useful for cases where we can't use a peak detector with asymptotic tracking.
 	 * For example: lookahead limiters, silence detection, etc.
 	 */
-	class SIGNALPROCESSING_API FMovingAverager
+	class FMovingAverager
 	{
 	public:
 		// Delay length in samples.
-		FMovingAverager(uint32 NumSamples);
+		SIGNALPROCESSING_API FMovingAverager(uint32 NumSamples);
 
 		// Returns average amplitude across the internal buffer, and fills Output with the delay line output.
-		float ProcessInput(const float& Input, float& Output);
+		SIGNALPROCESSING_API float ProcessInput(const float& Input, float& Output);
 
-		void SetNumSamples(uint32 NumSamples);
+		SIGNALPROCESSING_API void SetNumSamples(uint32 NumSamples);
 
 	private:
 		FMovingAverager();
@@ -40,14 +40,14 @@ namespace Audio
 	/**
 	 * Vectorized version of FMovingAverager.
 	 */
-	class SIGNALPROCESSING_API FMovingVectorAverager
+	class FMovingVectorAverager
 	{
 	public:
 		// Delay length in samples. NumSamples must be divisible by four.
-		FMovingVectorAverager(uint32 NumSamples);
+		SIGNALPROCESSING_API FMovingVectorAverager(uint32 NumSamples);
 
 		// Returns average amplitude across the internal buffer, and fills Output with the delay line output.
-		float ProcessAudio(const VectorRegister4Float& Input, VectorRegister4Float& Output);
+		SIGNALPROCESSING_API float ProcessAudio(const VectorRegister4Float& Input, VectorRegister4Float& Output);
 
 	private:
 		FMovingVectorAverager();
@@ -65,26 +65,26 @@ namespace Audio
 	 * This object will return buffered audio while the input signal is louder than the specified threshold,
 	 * and buffer audio when the input signal otherwise.
 	 */
-	class SIGNALPROCESSING_API FSilenceDetection
+	class FSilenceDetection
 	{
 	public:
 		// InOnsetThreshold is the minimum amplitude of a signal before we begin outputting audio, in linear gain.
 		// InReleaseThreshold is the amplitude of the signal before we stop outputting audio, in linear gain.
 		// AttackDurationInSamples is the amount of samples we average over when calculating our amplitude when the in audio is below the threshold.
 		// ReleaseDurationInSamples is the amount of samples we average over when calculating our amplitude when the input audio is above the threshold.
-		FSilenceDetection(float InOnsetThreshold, float InReleaseThreshold, int32 AttackDurationInSamples, int32 ReleaseDurationInSamples);
+		SIGNALPROCESSING_API FSilenceDetection(float InOnsetThreshold, float InReleaseThreshold, int32 AttackDurationInSamples, int32 ReleaseDurationInSamples);
 
 		// Buffers InAudio and renders any non-silent audio to OutAudio. Returns the number of samples written to OutAudio.
 		// The number of samples returned will only be less than NumSamples if the signal becomes audible mid-buffer.
 		// We do not return partial buffers when returning from an audible state to a silent state.
 		// This should also work in place, i.e. if InAudio == OutAudio.
-		int32 ProcessBuffer(const float* InAudio, float* OutAudio, int32 NumSamples);
+		SIGNALPROCESSING_API int32 ProcessBuffer(const float* InAudio, float* OutAudio, int32 NumSamples);
 
 		// Set the threshold of audibility, in terms of linear gain.
-		void SetThreshold(float InThreshold);
+		SIGNALPROCESSING_API void SetThreshold(float InThreshold);
 
 		// Returns the current estimate of the current amplitude of the input signal, in linear gain.
-		float GetCurrentAmplitude();
+		SIGNALPROCESSING_API float GetCurrentAmplitude();
 	private:
 		FSilenceDetection();
 
@@ -100,21 +100,21 @@ namespace Audio
 	 * This object accepts an input buffer and current amplitude estimate of that input buffer,
 	 * Then applies a computed gain target. Works like a standard feed forward limiter, with a threshold of 0.
 	 */
-	class SIGNALPROCESSING_API FSlowAdaptiveGainControl
+	class FSlowAdaptiveGainControl
 	{
 	public:
 		// InGainTarget is our target running linear gain.
 		// InAdaptiveRate is the time it will take to respond to changes in amplitude, in numbers of buffer callbacks.
 		// InGainMin is the most we will attenuate the input signal.
 		// InGainMax is the most we will amplify the input signal.
-		FSlowAdaptiveGainControl(float InGainTarget, int32 InAdaptiveRate, float InGainMin = 0.01f, float InGainMax = 2.0f);
+		SIGNALPROCESSING_API FSlowAdaptiveGainControl(float InGainTarget, int32 InAdaptiveRate, float InGainMin = 0.01f, float InGainMax = 2.0f);
 		
 		// Takes an amplitude estimate and an input buffer, and attenuates InAudio based on it.
 		// Returns the linear gain applied to InAudio.
-		float ProcessAudio(float* InAudio, int32 NumSamples, float InAmplitude);
+		SIGNALPROCESSING_API float ProcessAudio(float* InAudio, int32 NumSamples, float InAmplitude);
 
 		// Sets the responsiveness of the adaptive gain control, in number of buffer callbacks.
-		void SetAdaptiveRate(int32 InAdaptiveRate);
+		SIGNALPROCESSING_API void SetAdaptiveRate(int32 InAdaptiveRate);
 
 	private:
 		FSlowAdaptiveGainControl();
@@ -140,17 +140,17 @@ namespace Audio
 	 * If no new weights are given and we've converged to our previous input weights, this works like a normal FFT-based FIR filter.
 	 * Convergence is non-asymptotic: if no new weights are given after our number of steps until convergence, our filter is using the exact weights given.
 	 */
-	class SIGNALPROCESSING_API FAdaptiveFilter_DEPRECATED // Deprecated in 5.1
+	class FAdaptiveFilter_DEPRECATED // Deprecated in 5.1
 	{
 	public:
-		FAdaptiveFilter_DEPRECATED(int32 FilterLength, int32 AudioCallbackSize);
+		SIGNALPROCESSING_API FAdaptiveFilter_DEPRECATED(int32 FilterLength, int32 AudioCallbackSize);
 
 		/*
 		* Applies current filter to InAudio in-place. If there is a new target set of weights, they can be input below.
 		*/
-		void ProcessAudio(float* InAudio, int32 NumSamples);
+		SIGNALPROCESSING_API void ProcessAudio(float* InAudio, int32 NumSamples);
 
-		void SetWeights(const FrequencyBuffer& InFilterWeights, int32 FilterLength, float InLearningRate);
+		SIGNALPROCESSING_API void SetWeights(const FrequencyBuffer& InFilterWeights, int32 FilterLength, float InLearningRate);
 
 	private:
 
@@ -178,7 +178,7 @@ namespace Audio
 	};
 
 	class UE_DEPRECATED(5.1, "FAdaptiveFilter will no longer be supported.") FAdaptiveFilter;
-	class SIGNALPROCESSING_API FAdaptiveFilter : public FAdaptiveFilter_DEPRECATED
+	class FAdaptiveFilter : public FAdaptiveFilter_DEPRECATED
 	{
 		public:
 			using FAdaptiveFilter_DEPRECATED::FAdaptiveFilter_DEPRECATED;
@@ -187,12 +187,12 @@ namespace Audio
 	/**
 	 * This class takes an incoming signal and an outgoing signal, Correlates them, and returns the frequency values of the weight targets to pass to an adaptive filter.  
 	 */
-	class SIGNALPROCESSING_API FFDAPFilterComputer_DEPRECATED // Deprecated in 5.1
+	class FFDAPFilterComputer_DEPRECATED // Deprecated in 5.1
 	{
 	public:
-		FFDAPFilterComputer_DEPRECATED();
+		SIGNALPROCESSING_API FFDAPFilterComputer_DEPRECATED();
 
-		void GenerateWeights(const float* IncomingSignal, int32 NumIncomingSamples, const float* OutgoingSignal, int32 NumOutgoingSamples, FrequencyBuffer& OutWeights);
+		SIGNALPROCESSING_API void GenerateWeights(const float* IncomingSignal, int32 NumIncomingSamples, const float* OutgoingSignal, int32 NumOutgoingSamples, FrequencyBuffer& OutWeights);
 
 	private:
 		FrequencyBuffer IncomingFrequencies;
@@ -203,7 +203,7 @@ namespace Audio
 	};
 
 	class UE_DEPRECATED(5.1, "FAFDAPFFilterComputer will no longer be supported.") FFDAPFilterComputer;
-	class SIGNALPROCESSING_API FFDAPFilterComputer : public FFDAPFilterComputer_DEPRECATED 
+	class FFDAPFilterComputer : public FFDAPFilterComputer_DEPRECATED 
 	{
 	};
 
@@ -212,20 +212,20 @@ namespace Audio
 	 * To add a new patch to a rendered audio signal, user AddNewSignalPatch. See FPatchInput for how to push audio.
 	 * ProcessAudio then filters the microphone signal accordingly.
 	 */
-	class SIGNALPROCESSING_API FAcousticEchoCancellation_DEPRECATED // Deprecated in 5.1
+	class FAcousticEchoCancellation_DEPRECATED // Deprecated in 5.1
 	{
 	public:
 		/**
 		 * Convergence Rate should be a number between 0 and 1. The higher the number, the quicker the adaptive filter reacts. 
 		 */
-		FAcousticEchoCancellation_DEPRECATED(float InConvergenceRate, int32 CallbackSize, int32 InFilterLength, int32 InFilterUpdateRate = 1);
+		SIGNALPROCESSING_API FAcousticEchoCancellation_DEPRECATED(float InConvergenceRate, int32 CallbackSize, int32 InFilterLength, int32 InFilterUpdateRate = 1);
 
 		/** Callback function for outgoing audio signal. This is where the filter is applied, and the bulk of the DSP work takes place. */
-		void ProcessAudio(float* InAudio, int32 NumSamples);
+		SIGNALPROCESSING_API void ProcessAudio(float* InAudio, int32 NumSamples);
 
 		/** This is how any signal that may be picked up by the microphone may be added to the echo cancellation here: */
-		FPatchInput AddNewSignalPatch(int32 ExpectedLatency, float Gain = 1.0f);
-		void RemoveSignalPatch(const FPatchInput& Patch);
+		SIGNALPROCESSING_API FPatchInput AddNewSignalPatch(int32 ExpectedLatency, float Gain = 1.0f);
+		SIGNALPROCESSING_API void RemoveSignalPatch(const FPatchInput& Patch);
 
 	private:
 		FPatchMixer PatchMixer;
@@ -241,7 +241,7 @@ namespace Audio
 	};
 
 	class UE_DEPRECATED(5.1, "FAcousticEchoCancellation will no longer be supported.") FAcousticEchoCancellation;
-	class SIGNALPROCESSING_API FAcousticEchoCancellation : public FAcousticEchoCancellation_DEPRECATED
+	class FAcousticEchoCancellation : public FAcousticEchoCancellation_DEPRECATED
 	{
 		public:
 			using FAcousticEchoCancellation_DEPRECATED::FAcousticEchoCancellation_DEPRECATED;
diff --git a/Engine/Source/Runtime/SignalProcessing/Public/DSP/VolumeFader.h b/Engine/Source/Runtime/SignalProcessing/Public/DSP/VolumeFader.h
--- a/Engine/Source/Runtime/SignalProcessing/Public/DSP/VolumeFader.h
+++ b/Engine/Source/Runtime/SignalProcessing/Public/DSP/VolumeFader.h
@@ -26,76 +26,76 @@ namespace Audio
 
 
 	/** Control-rate fader for managing volume fades of various standard shapes. */
-	class SIGNALPROCESSING_API FVolumeFader
+	class FVolumeFader
 	{
 	public:
-		FVolumeFader();
+		SIGNALPROCESSING_API FVolumeFader();
 
 		/**
 		 * Activates the fader if currently deactivated. Sets duration to indefinite (-1.0f)
 		 */
-		void Activate();
+		SIGNALPROCESSING_API void Activate();
 
 		/**
 		 * Deactivates the fader, causing it to remain
 		 * at the current value and disregard update.
 		 */
-		void Deactivate();
+		SIGNALPROCESSING_API void Deactivate();
 
 		/**
 		 * Gets time fader is to remain active (-1.0f if always active).
 		 */
-		float GetActiveDuration() const;
+		SIGNALPROCESSING_API float GetActiveDuration() const;
 
 		/**
 		 * Returns current volume of fader
 		 */
-		float GetVolume() const;
+		SIGNALPROCESSING_API float GetVolume() const;
 
 		/**
 		 * Returns the volume given the delta from the current time
 		 * into the future (Effectively like running to retrieve volume
 		 * but without actually updating internal state).
 		 */
-		float GetVolumeAfterTime(float InDeltaTime) const;
+		SIGNALPROCESSING_API float GetVolumeAfterTime(float InDeltaTime) const;
 
 		/**
 		 * Returns the duration of the fade.
 		 */
-		float GetFadeDuration() const;
+		SIGNALPROCESSING_API float GetFadeDuration() const;
 
 		/**
 		 * Returns the curve type of the fader
 		 */
-		EFaderCurve GetCurve() const;
+		SIGNALPROCESSING_API EFaderCurve GetCurve() const;
 
 		/**
 		 * Returns the target volume of the fader
 		 */
-		float GetTargetVolume() const;
+		SIGNALPROCESSING_API float GetTargetVolume() const;
 
 		/**
 		 * Whether or not the fader is active.
 		 */
-		bool IsActive() const;
+		SIGNALPROCESSING_API bool IsActive() const;
 
 		/**
 		 * Returns whether or not the fader is currently
 		 * fading over time.
 		 */
-		bool IsFading() const;
+		SIGNALPROCESSING_API bool IsFading() const;
 
 		/**
 		 * Returns whether or not the fader is currently
 		 * fading over time and value is increasing.
 		 */
-		bool IsFadingIn() const;
+		SIGNALPROCESSING_API bool IsFadingIn() const;
 
 		/**
 		 * Returns whether or not the fader is currently
 		 * fading over time and value is decreasing.
 		 */
-		bool IsFadingOut() const;
+		SIGNALPROCESSING_API bool IsFadingOut() const;
 
 		/**
 		 * Sets the duration the fader is to be active in the future,
@@ -104,27 +104,27 @@ namespace Audio
 		 * at which point it will continue applying the fade over the
 		 * remainder of the fade duration.
 		 */
-		void SetActiveDuration(float InDuration);
+		SIGNALPROCESSING_API void SetActiveDuration(float InDuration);
 
 		/**
 		 * Sets the volume immediately, interrupting any currently active fade.
 		 */
-		void SetVolume(float InVolume);
+		SIGNALPROCESSING_API void SetVolume(float InVolume);
 
 		/**
 		 * Applies a volume fade over time with the provided parameters.
 		 */
-		void StartFade(float InVolume, float InDuration, EFaderCurve InCurve);
+		SIGNALPROCESSING_API void StartFade(float InVolume, float InDuration, EFaderCurve InCurve);
 
 		/**
 		 * Stops fade, maintaining the current value as the target.
 		 */
-		void StopFade();
+		SIGNALPROCESSING_API void StopFade();
 
 		/**
 		 * Updates the fader's state with the given delta in time since last update.
 		 */
-		void Update(float InDeltaTime);
+		SIGNALPROCESSING_API void Update(float InDeltaTime);
 
 	private:
 		/** Converts value to final resulting volume */
diff --git a/Engine/Source/Runtime/SignalProcessing/Public/DSP/WaveShaper.h b/Engine/Source/Runtime/SignalProcessing/Public/DSP/WaveShaper.h
--- a/Engine/Source/Runtime/SignalProcessing/Public/DSP/WaveShaper.h
+++ b/Engine/Source/Runtime/SignalProcessing/Public/DSP/WaveShaper.h
@@ -17,38 +17,38 @@ namespace Audio
 
 	// A digital wave shaping effect to cause audio distortion
 	// https://en.wikipedia.org/wiki/Waveshaper
-	class SIGNALPROCESSING_API FWaveShaper
+	class FWaveShaper
 	{
 	public:
 		// Constructor
-		FWaveShaper();
+		SIGNALPROCESSING_API FWaveShaper();
 
 		// Destructor
-		~FWaveShaper();
+		SIGNALPROCESSING_API ~FWaveShaper();
 
 		// Initialize the equalizer
-		void Init(const float InSampleRate);
+		SIGNALPROCESSING_API void Init(const float InSampleRate);
 
 		// Sets the amount of wave shapping. 0.0 is no effect.
-		void SetAmount(const float InAmount);
+		SIGNALPROCESSING_API void SetAmount(const float InAmount);
 
 		// Set DC offset before processing
-		void SetBias(const float InBias);
+		SIGNALPROCESSING_API void SetBias(const float InBias);
 
 		// Sets the output gain of the waveshaper
-		void SetOutputGainDb(const float InGainDb);
+		SIGNALPROCESSING_API void SetOutputGainDb(const float InGainDb);
 
 		// Sets the output gain of the waveshaper
-		void SetOutputGainLinear(const float InGainLinear);
+		SIGNALPROCESSING_API void SetOutputGainLinear(const float InGainLinear);
 
-		void SetType(const EWaveShaperType InType);
+		SIGNALPROCESSING_API void SetType(const EWaveShaperType InType);
 
 		// Processes one Sample, using ATan
 		// Deprecating in favor of ProcessAudioBuffer
-		void ProcessAudio(const float InSample, float& OutSample);
+		SIGNALPROCESSING_API void ProcessAudio(const float InSample, float& OutSample);
 
 		// Process an entire buffer of audio
-		void ProcessAudioBuffer(const float* InBuffer, float* OutBuffer, int32 NumFrames);
+		SIGNALPROCESSING_API void ProcessAudioBuffer(const float* InBuffer, float* OutBuffer, int32 NumFrames);
 
 	private:
 
diff --git a/Engine/Source/Runtime/SignalProcessing/Public/DSP/WaveTableOsc.h b/Engine/Source/Runtime/SignalProcessing/Public/DSP/WaveTableOsc.h
--- a/Engine/Source/Runtime/SignalProcessing/Public/DSP/WaveTableOsc.h
+++ b/Engine/Source/Runtime/SignalProcessing/Public/DSP/WaveTableOsc.h
@@ -27,48 +27,48 @@ namespace Audio
 
 
 	// A wave table oscillator class
-	class SIGNALPROCESSING_API FWaveTableOsc
+	class FWaveTableOsc
 	{
 	public:
 		// Constructor
-		FWaveTableOsc();
+		SIGNALPROCESSING_API FWaveTableOsc();
 
 		// Virtual Destructor
-		virtual ~FWaveTableOsc();
+		SIGNALPROCESSING_API virtual ~FWaveTableOsc();
 
 		// Initialize the wave table oscillator
-		void Init(const float InSampleRate, const float InFrequencyHz);
+		SIGNALPROCESSING_API void Init(const float InSampleRate, const float InFrequencyHz);
 
 		// Sets the sample rate of the oscillator.
-		void SetSampleRate(const float InSampleRate);
+		SIGNALPROCESSING_API void SetSampleRate(const float InSampleRate);
 
 		// Resets the wave table read indices.
-		void Reset();
+		SIGNALPROCESSING_API void Reset();
 
 		// Sets the amount to scale and add to the output of the wave table
-		void SetScaleAdd(const float InScale, const float InAdd);
+		SIGNALPROCESSING_API void SetScaleAdd(const float InScale, const float InAdd);
 
 		// Returns the type of the wave table oscillator.
 		EWaveTable::Type GetType() const { return WaveTableType; }
 
 		// Sets the frequency of the wave table oscillator.
-		void SetFrequencyHz(const float InFrequencyHz);
+		SIGNALPROCESSING_API void SetFrequencyHz(const float InFrequencyHz);
 
 		// Returns the frequency of the wave table oscillator.
 		float GetFrequencyHz() const { return FrequencyHz; }
 
 		// Returns the internal table used in the wave table.
-		TArray<float>& GetTable();
-		const TArray<float>& GetTable() const;
+		SIGNALPROCESSING_API TArray<float>& GetTable();
+		SIGNALPROCESSING_API const TArray<float>& GetTable() const;
 
 		// Processes the wave table, outputs the normal and quad phase (optional) values 
-		void Generate(float* OutputNormalPhase, float* OutputQuadPhase = nullptr);
+		SIGNALPROCESSING_API void Generate(float* OutputNormalPhase, float* OutputQuadPhase = nullptr);
 
 		// Creates a wave table using internal factories for standard wave tables or uses custom wave table factor if it exists.
-		static TSharedPtr<FWaveTableOsc> CreateWaveTable(const EWaveTable::Type WaveTableType, const int32 WaveTableSize = 1024);
+		static SIGNALPROCESSING_API TSharedPtr<FWaveTableOsc> CreateWaveTable(const EWaveTable::Type WaveTableType, const int32 WaveTableSize = 1024);
 
 	protected:
-		void UpdateFrequency();
+		SIGNALPROCESSING_API void UpdateFrequency();
 
 		// The wave table buffer
 		TArray<float> WaveTableBuffer;
diff --git a/Engine/Source/Runtime/SkeletalMeshDescription/Public/SkeletalMeshAttributes.h b/Engine/Source/Runtime/SkeletalMeshDescription/Public/SkeletalMeshAttributes.h
--- a/Engine/Source/Runtime/SkeletalMeshDescription/Public/SkeletalMeshAttributes.h
+++ b/Engine/Source/Runtime/SkeletalMeshDescription/Public/SkeletalMeshAttributes.h
@@ -34,7 +34,7 @@ namespace MeshAttribute
 }
 
 
-class SKELETALMESHDESCRIPTION_API FSkeletalMeshAttributesShared
+class FSkeletalMeshAttributesShared
 {
 public:
 	using FBoneArray = TMeshElementContainer<FBoneID>;
@@ -58,51 +58,51 @@ public:
 	 * added to the  mesh description to represent Bones (similar to the Vertex/Polygons/Edges elements). 
 	 * MeshAttribute::Bone::Name is just one of the attributes of the Bones element.
 	 */
-	static FName BonesElementName;
+	static SKELETALMESHDESCRIPTION_API FName BonesElementName;
 
 	// The name of the default skin weight profile.
-	static FName DefaultSkinWeightProfileName;
+	static SKELETALMESHDESCRIPTION_API FName DefaultSkinWeightProfileName;
 	
-	FSkeletalMeshAttributesShared(const FMeshDescription& InMeshDescription);
+	SKELETALMESHDESCRIPTION_API FSkeletalMeshAttributesShared(const FMeshDescription& InMeshDescription);
 
 	//
 	// Skin Weights Methods
 	//
 
 	/// Returns the list of all registered skin weight profile names on this mesh.
-	TArray<FName> GetSkinWeightProfileNames() const;
+	SKELETALMESHDESCRIPTION_API TArray<FName> GetSkinWeightProfileNames() const;
 	
 	/// Returns \c true if the given identifier is a valid profile name. If the name is empty, or matches the default profile,
 	/// then the profile name is considered invalid. 
-	static bool IsValidSkinWeightProfileName(const FName& InProfileName);
+	static SKELETALMESHDESCRIPTION_API bool IsValidSkinWeightProfileName(const FName& InProfileName);
 
 	/// Helper function that indicates whether an attribute name represents a skin weight attribute.
-	static bool IsSkinWeightAttribute(const FName& InAttributeName);
+	static SKELETALMESHDESCRIPTION_API bool IsSkinWeightAttribute(const FName& InAttributeName);
 
 	/// Returns a skin profile name from the attribute name, if the attribute name is a valid skin weights
 	/// attribute.
-	static FName GetProfileNameFromAttribute(const FName& InAttributeName);
+	static SKELETALMESHDESCRIPTION_API FName GetProfileNameFromAttribute(const FName& InAttributeName);
 
-	FSkinWeightsVertexAttributesConstRef GetVertexSkinWeights(const FName& InProfileName = NAME_None) const;
+	SKELETALMESHDESCRIPTION_API FSkinWeightsVertexAttributesConstRef GetVertexSkinWeights(const FName& InProfileName = NAME_None) const;
 	
-	FSkinWeightsVertexAttributesConstRef GetVertexSkinWeightsFromAttributeName(const FName& InAttributeName = NAME_None) const;
+	SKELETALMESHDESCRIPTION_API FSkinWeightsVertexAttributesConstRef GetVertexSkinWeightsFromAttributeName(const FName& InAttributeName = NAME_None) const;
 
 
 	//
 	// Bones Methods
 	//
 
-	bool HasBoneColorAttribute() const;
+	SKELETALMESHDESCRIPTION_API bool HasBoneColorAttribute() const;
 
-	bool HasBoneNameAttribute() const;
+	SKELETALMESHDESCRIPTION_API bool HasBoneNameAttribute() const;
 
-	bool HasBonePoseAttribute() const;
+	SKELETALMESHDESCRIPTION_API bool HasBonePoseAttribute() const;
 
-	bool HasBoneParentIndexAttribute() const;
+	SKELETALMESHDESCRIPTION_API bool HasBoneParentIndexAttribute() const;
 
-	const FBoneArray& Bones() const;
+	SKELETALMESHDESCRIPTION_API const FBoneArray& Bones() const;
 
-	const TAttributesSet<FBoneID>& BoneAttributes() const;
+	SKELETALMESHDESCRIPTION_API const TAttributesSet<FBoneID>& BoneAttributes() const;
 
 	/**  @return true, if bone element was added to the MeshDescription. */
 	bool HasBones() const
@@ -111,24 +111,24 @@ public:
 	}
 
 	/**  @return the number of the bones. 0 if bone element does not exist in the MeshDescripion. */
-	int32 GetNumBones() const;
+	SKELETALMESHDESCRIPTION_API int32 GetNumBones() const;
 
 	/** @return true, if the passed bone ID is valid */
-	bool IsBoneValid(const FBoneID BoneID) const;
+	SKELETALMESHDESCRIPTION_API bool IsBoneValid(const FBoneID BoneID) const;
 
-	FBoneNameAttributesConstRef GetBoneNames() const;
+	SKELETALMESHDESCRIPTION_API FBoneNameAttributesConstRef GetBoneNames() const;
 
-	FBoneParentIndexAttributesConstRef GetBoneParentIndices() const;
+	SKELETALMESHDESCRIPTION_API FBoneParentIndexAttributesConstRef GetBoneParentIndices() const;
 
-	FBonePoseAttributesConstRef GetBonePoses() const;
+	SKELETALMESHDESCRIPTION_API FBonePoseAttributesConstRef GetBonePoses() const;
 
-	FBoneColorAttributesConstRef GetBoneColors() const;
+	SKELETALMESHDESCRIPTION_API FBoneColorAttributesConstRef GetBoneColors() const;
 
 protected:
 	/// Construct a name for a skin weight attribute with the given skin weight profile name.
 	/// Each mesh description can hold different skin weight profiles, although the default
 	/// is always present.
-	static FName CreateSkinWeightAttributeName(const FName& InProfileName);
+	static SKELETALMESHDESCRIPTION_API FName CreateSkinWeightAttributeName(const FName& InProfileName);
 
 private:
 	const FMeshElementChannels* BoneElementsShared = nullptr;
@@ -137,15 +137,15 @@ private:
 };
 
 
-class SKELETALMESHDESCRIPTION_API FSkeletalMeshAttributes :
+class FSkeletalMeshAttributes :
 	public FStaticMeshAttributes,
 	public FSkeletalMeshAttributesShared
 {
 public:
 
-	explicit FSkeletalMeshAttributes(FMeshDescription& InMeshDescription);
+	SKELETALMESHDESCRIPTION_API explicit FSkeletalMeshAttributes(FMeshDescription& InMeshDescription);
 
-	virtual void Register(bool bKeepExistingAttribute = false) override;
+	SKELETALMESHDESCRIPTION_API virtual void Register(bool bKeepExistingAttribute = false) override;
 	
 	//
 	// Skin Weights Methods
@@ -154,7 +154,7 @@ public:
 	/// Register a new skin weight profile with the given name. The attribute name will encode the profile name and
 	/// it will be listed in GetSkinWeightProfileNames(). Returns \c true if the profile was successfully registered.
 	/// Returns \c false if the attribute was already registered or if IsValidSkinWeightProfileName() returned false.
-	bool RegisterSkinWeightAttribute(const FName& InProfileName);
+	SKELETALMESHDESCRIPTION_API bool RegisterSkinWeightAttribute(const FName& InProfileName);
 
 	static bool IsReservedAttributeName(const FName& InAttributeName)
 	{
@@ -167,40 +167,40 @@ public:
 	}
 	
 	/// Returns the skin weight profile given by its name. NAME_None corresponds to the default profile.
-	FSkinWeightsVertexAttributesRef GetVertexSkinWeights(const FName& InProfileName = NAME_None);
+	SKELETALMESHDESCRIPTION_API FSkinWeightsVertexAttributesRef GetVertexSkinWeights(const FName& InProfileName = NAME_None);
 
-	FSkinWeightsVertexAttributesRef GetVertexSkinWeightsFromAttributeName(const FName& InAttributeName = NAME_None);
+	SKELETALMESHDESCRIPTION_API FSkinWeightsVertexAttributesRef GetVertexSkinWeightsFromAttributeName(const FName& InAttributeName = NAME_None);
 
 	//
 	// Bones Methods
 	//
 
 	/** Register an optional color attribute for bones */
-	void RegisterColorAttribute();
+	SKELETALMESHDESCRIPTION_API void RegisterColorAttribute();
 
-	FBoneArray& Bones();
+	SKELETALMESHDESCRIPTION_API FBoneArray& Bones();
 
-	TAttributesSet<FBoneID>& BoneAttributes();
+	SKELETALMESHDESCRIPTION_API TAttributesSet<FBoneID>& BoneAttributes();
 
 	/** Reserves space for this number of new bones */
-	void ReserveNewBones(const int InNumBones);
+	SKELETALMESHDESCRIPTION_API void ReserveNewBones(const int InNumBones);
 
 	/** Adds a new bone and returns its ID */
-	FBoneID CreateBone();
+	SKELETALMESHDESCRIPTION_API FBoneID CreateBone();
 
 	/** Adds a new bone  with the given ID */
-	void CreateBone(const FBoneID BoneID);
+	SKELETALMESHDESCRIPTION_API void CreateBone(const FBoneID BoneID);
 
 	/** Deletes a bone with the given ID */
-	void DeleteBone(const FBoneID BoneID);
+	SKELETALMESHDESCRIPTION_API void DeleteBone(const FBoneID BoneID);
 
-	FBoneNameAttributesRef GetBoneNames();
+	SKELETALMESHDESCRIPTION_API FBoneNameAttributesRef GetBoneNames();
 	
-	FBoneParentIndexAttributesRef GetBoneParentIndices();
+	SKELETALMESHDESCRIPTION_API FBoneParentIndexAttributesRef GetBoneParentIndices();
 
-	FBonePoseAttributesRef GetBonePoses();
+	SKELETALMESHDESCRIPTION_API FBonePoseAttributesRef GetBonePoses();
 
-	FBoneColorAttributesRef GetBoneColors();
+	SKELETALMESHDESCRIPTION_API FBoneColorAttributesRef GetBoneColors();
 
 private:
 	FMeshElementChannels* BoneElements = nullptr;
diff --git a/Engine/Source/Runtime/SkeletalMeshDescription/Public/SkeletalMeshOperations.h b/Engine/Source/Runtime/SkeletalMeshDescription/Public/SkeletalMeshOperations.h
--- a/Engine/Source/Runtime/SkeletalMeshDescription/Public/SkeletalMeshOperations.h
+++ b/Engine/Source/Runtime/SkeletalMeshDescription/Public/SkeletalMeshOperations.h
@@ -17,7 +17,7 @@ DECLARE_LOG_CATEGORY_EXTERN(LogSkeletalMeshOperations, Log, All);
 
 
 
-class SKELETALMESHDESCRIPTION_API FSkeletalMeshOperations : public FStaticMeshOperations
+class FSkeletalMeshOperations : public FStaticMeshOperations
 {
 public:
 	struct FSkeletalMeshAppendSettings
@@ -30,5 +30,5 @@ public:
 		TArray<FBoneIndexType> SourceRemapBoneIndex;
 	};
 	
-	static void AppendSkinWeight(const FMeshDescription& SourceMesh, FMeshDescription& TargetMesh, FSkeletalMeshAppendSettings& AppendSettings);
+	static SKELETALMESHDESCRIPTION_API void AppendSkinWeight(const FMeshDescription& SourceMesh, FMeshDescription& TargetMesh, FSkeletalMeshAppendSettings& AppendSettings);
 };
diff --git a/Engine/Source/Runtime/Sockets/Public/IPAddress.h b/Engine/Source/Runtime/Sockets/Public/IPAddress.h
--- a/Engine/Source/Runtime/Sockets/Public/IPAddress.h
+++ b/Engine/Source/Runtime/Sockets/Public/IPAddress.h
@@ -18,7 +18,7 @@
 /**
  * Represents an internet address. All data is in network byte order
  */
-class SOCKETS_API FInternetAddr
+class FInternetAddr
 {
 protected:
 	/** Hidden on purpose */
@@ -183,7 +183,7 @@ public:
 		return NAME_None;
 	}
 
-	virtual void DumpAddrData() const;
+	SOCKETS_API virtual void DumpAddrData() const;
 
 	friend inline uint32 GetTypeHash(const FInternetAddr& InAddr)
 	{
diff --git a/Engine/Source/Runtime/Sockets/Public/NetworkMessage.h b/Engine/Source/Runtime/Sockets/Public/NetworkMessage.h
--- a/Engine/Source/Runtime/Sockets/Public/NetworkMessage.h
+++ b/Engine/Source/Runtime/Sockets/Public/NetworkMessage.h
@@ -85,7 +85,7 @@ ENUM_CLASS_FLAGS(EConnectionFlags);
 /**
 * Simple abstraction for sockets that allows FNFSMessageHeader to use either an ordinary socket or a mutichannel socket
 **/
-class SOCKETS_API FSimpleAbstractSocket
+class FSimpleAbstractSocket
 {
 public:
 	/**
@@ -111,7 +111,7 @@ public:
 /**
 * Ordinary socket version of FSimpleAbstractSocket
 **/
-class SOCKETS_API FSimpleAbstractSocket_FSocket : public FSimpleAbstractSocket
+class FSimpleAbstractSocket_FSocket : public FSimpleAbstractSocket
 {
 	/** Ordinary socket to forward requests to **/
 	class FSocket* Socket;
@@ -124,8 +124,8 @@ public:
 		: Socket(InSocket)
 	{
 	}
-	virtual bool Receive(uint8 *Results, int32 Size) const;
-	virtual bool Send(const uint8 *Buffer, int32 Size) const;
+	SOCKETS_API virtual bool Receive(uint8 *Results, int32 Size) const;
+	SOCKETS_API virtual bool Send(const uint8 *Buffer, int32 Size) const;
 	virtual uint32 GetMagic() const
 	{
 		return 0x9E2B83C1;
@@ -135,7 +135,7 @@ public:
 /**
 * Multichannel socket version of FSimpleAbstractSocket
 **/
-class SOCKETS_API FSimpleAbstractSocket_FMultichannelTCPSocket : public FSimpleAbstractSocket
+class FSimpleAbstractSocket_FMultichannelTCPSocket : public FSimpleAbstractSocket
 {
 	/** Multichannel socket to forward requests to **/
 	class FMultichannelTcpSocket* Socket;
@@ -158,8 +158,8 @@ public:
 		check(SendChannel);
 		check(ReceiveChannel);
 	}
-	virtual bool Receive(uint8 *Results, int32 Size) const;
-	virtual bool Send(const uint8 *Buffer, int32 Size) const;
+	SOCKETS_API virtual bool Receive(uint8 *Results, int32 Size) const;
+	SOCKETS_API virtual bool Send(const uint8 *Buffer, int32 Size) const;
 	virtual uint32 GetMagic() const
 	{
 		return 0x9E2B83C2;
@@ -169,7 +169,7 @@ public:
 /**
 * Simple wrapper for sending and receiving atomic packets
 **/
-struct SOCKETS_API FNFSMessageHeader
+struct FNFSMessageHeader
 {
 	/** Magic number, used for error checking and endianess checking **/
 	uint32 Magic;
@@ -233,7 +233,7 @@ struct SOCKETS_API FNFSMessageHeader
 	 *
 	 * @return true if successful
 	 */
-	static bool WrapAndSendPayload(const TArray<uint8>& Payload, const FSimpleAbstractSocket& Socket);
+	static SOCKETS_API bool WrapAndSendPayload(const TArray<uint8>& Payload, const FSimpleAbstractSocket& Socket);
 
 	/**
 	 * This function will receive a header, and then the payload array from the network
@@ -243,7 +243,7 @@ struct SOCKETS_API FNFSMessageHeader
 	 *
 	 * @return true if successful
 	 */
-	static bool ReceivePayload(FArrayReader& OutPayload, const FSimpleAbstractSocket& Socket);
+	static SOCKETS_API bool ReceivePayload(FArrayReader& OutPayload, const FSimpleAbstractSocket& Socket);
 
 	/**
 	 * This function will send a payload data (with header) and wait for a response, serializing
@@ -255,7 +255,7 @@ struct SOCKETS_API FNFSMessageHeader
 	 *
 	 * @return true if successful
 	 */
-	static bool SendPayloadAndReceiveResponse(const TArray<uint8>& Payload, class FArrayReader& Response, const FSimpleAbstractSocket& Socket);
+	static SOCKETS_API bool SendPayloadAndReceiveResponse(const TArray<uint8>& Payload, class FArrayReader& Response, const FSimpleAbstractSocket& Socket);
 
 };
 
@@ -263,7 +263,7 @@ struct SOCKETS_API FNFSMessageHeader
 /**
  * A helper class for storing all available file info.
  */
-struct SOCKETS_API FFileInfo
+struct FFileInfo
 {
 	bool FileExists;
 	bool ReadOnly;
@@ -284,7 +284,7 @@ struct SOCKETS_API FFileInfo
 /**
  * A helper class for wrapping some of the network file payload specifics
  */
-class SOCKETS_API FNetworkFileArchive : public FBufferArchive
+class FNetworkFileArchive : public FBufferArchive
 {
 public:
 	FNetworkFileArchive(uint32 Command)
diff --git a/Engine/Source/Runtime/Sockets/Public/ServerTOC.h b/Engine/Source/Runtime/Sockets/Public/ServerTOC.h
--- a/Engine/Source/Runtime/Sockets/Public/ServerTOC.h
+++ b/Engine/Source/Runtime/Sockets/Public/ServerTOC.h
@@ -11,7 +11,7 @@
 /**
  * Helper struct to hold and quickly access server TOC.
  */
-struct SOCKETS_API FServerTOC
+struct FServerTOC
 {
 	/** List of files in a directory. */
 	typedef TMap<FString, FDateTime> FDirectory;
@@ -20,7 +20,7 @@ struct SOCKETS_API FServerTOC
 	TMap<FString, FDirectory*> Directories;
 
 	/** Destructor. Destroys directories. */
-	~FServerTOC();
+	SOCKETS_API ~FServerTOC();
 
 	/**
 	 * Adds a file or directory to TOC.
@@ -28,7 +28,7 @@ struct SOCKETS_API FServerTOC
 	 * @param Filename File name or directory name to add.
 	 * @param Timestamp File timestamp. Directories should have this set to 0.
 	 */
-	void AddFileOrDirectory(const FString& Filename, const FDateTime& Timestamp);
+	SOCKETS_API void AddFileOrDirectory(const FString& Filename, const FDateTime& Timestamp);
 
 	/**
 	 * Finds a file in TOC.
@@ -36,7 +36,7 @@ struct SOCKETS_API FServerTOC
 	 * @param Filename File name to find.
 	 * @return Pointer to a timestamp if the file was found, NULL otherwise.
 	 */
-	const FDateTime* FindFile(const FString& Filename) const;
+	SOCKETS_API const FDateTime* FindFile(const FString& Filename) const;
 
 	/**
 	 * Finds a directory in TOC.
@@ -44,7 +44,7 @@ struct SOCKETS_API FServerTOC
 	 * @param Directory Directory to find.
 	 * @return Pointer to a FDirectory if the directory was found, NULL otherwise.
 	 */
-	const FDirectory* FindDirectory(const FString& Directory) const;
+	SOCKETS_API const FDirectory* FindDirectory(const FString& Directory) const;
 
 	/**
 	 * Finds a directory in TOC non const version used internally 
@@ -53,8 +53,8 @@ struct SOCKETS_API FServerTOC
 	 * @param Directory Directory to find
 	 * @return pointer to a FDirectory if the directory was found, null otherwise
 	 */
-	FDirectory* FindDirectory(const FString& Directory);
+	SOCKETS_API FDirectory* FindDirectory(const FString& Directory);
 
 
-	int32 RemoveFileOrDirectory(const FString& Filename);
+	SOCKETS_API int32 RemoveFileOrDirectory(const FString& Filename);
 };
