speculative fix for quad delay crash

#rnx

diff --git a/Engine/Source/Runtime/SignalProcessing/Private/InterpolatedOnePole.cpp b/Engine/Source/Runtime/SignalProcessing/Private/InterpolatedOnePole.cpp
--- a/Engine/Source/Runtime/SignalProcessing/Private/InterpolatedOnePole.cpp
+++ b/Engine/Source/Runtime/SignalProcessing/Private/InterpolatedOnePole.cpp
@@ -94,6 +94,25 @@ namespace Audio
 		}
 	}
 
+	void FInterpolatedLPF::ProcessBufferInPlace(float* InOutBuffer, const int32 NumSamples)
+	{
+		for (int32 SampleIndex = 0; SampleIndex < NumSamples; ++SampleIndex)
+		{
+			// cache which delay term we should be using
+			const int32 ChannelIndex = SampleIndex % NumChannels;
+
+			// step forward coefficient
+			// Multiply delta by !ChannelIndex so the coefficient only accumulates at the beginning of each frame (on channel 0)
+			B1Curr += B1Delta * !ChannelIndex;
+
+			const float InputSample = InOutBuffer[SampleIndex];
+			float Yn = InputSample + B1Curr * (Z1Data[ChannelIndex] - InputSample); // LPF
+			Yn = UnderflowClamp(Yn);
+			Z1Data[ChannelIndex] = Yn;
+			InOutBuffer[SampleIndex] = Yn;
+		}
+	}
+
 	void FInterpolatedLPF::Reset()
 	{
 		B1Curr = 0.0f;
diff --git a/Engine/Source/Runtime/SignalProcessing/Public/DSP/InterpolatedOnePole.h b/Engine/Source/Runtime/SignalProcessing/Public/DSP/InterpolatedOnePole.h
--- a/Engine/Source/Runtime/SignalProcessing/Public/DSP/InterpolatedOnePole.h
+++ b/Engine/Source/Runtime/SignalProcessing/Public/DSP/InterpolatedOnePole.h
@@ -37,6 +37,9 @@ namespace Audio
 		// interpolates coefficient and processes a buffer
 		SIGNALPROCESSING_API void ProcessAudioBuffer(const float* RESTRICT InputBuffer, float* RESTRICT OutputBuffer, const int32 NumSamples);
 
+		// interpolates coefficient and processes a buffer
+		SIGNALPROCESSING_API void ProcessBufferInPlace(float* InOutBuffer, int32 NumSamples);
+
 		/*
 			StopFrequencyInterpolation() needs to be called manually when the interpolation should be done.
 			Snaps the coefficient to the target value.
