Fixing hard-coded scales/chords in midi note quantizer arrays.

diff --git a/Engine/Source/Runtime/SignalProcessing/Private/MidiNoteQuantizer.cpp b/Engine/Source/Runtime/SignalProcessing/Private/MidiNoteQuantizer.cpp
--- a/Engine/Source/Runtime/SignalProcessing/Private/MidiNoteQuantizer.cpp
+++ b/Engine/Source/Runtime/SignalProcessing/Private/MidiNoteQuantizer.cpp
@@ -109,7 +109,7 @@ namespace Audio
 		, { EMusicalScale::Scale::HalfDiminished_Locrian,
 			{{ Note::C1, Note::Db1, Note::Eb1, Note::F1, Note::Gb1, Note::Ab1, Note::Bb1 }, { Note::C1, Note::Eb1, Note::Gb1, Note::Bb1 }}}
 		, { EMusicalScale::Scale::Diminished,
-			{{ Note::C1, Note::D1, Note::Eb1, Note::F1, Note::Gb1, Note::Ab1, Note::A1, Note::B1 }, { Note::C1, Note::Eb1, Note::G1, Note::A1 }}}
+			{{ Note::C1, Note::D1, Note::Eb1, Note::F1, Note::Gb1, Note::Ab1, Note::A1, Note::B1 }, { Note::C1, Note::Eb1, Note::Gb1, Note::A1 }}}
 		, { EMusicalScale::Scale::MajorPentatonic,
 			{{ Note::C1, Note::D1, Note::E1, Note::G1, Note::A1}, { Note::C1, Note::E1, Note::G1 }}}
 		, { EMusicalScale::Scale::Lydian,
@@ -121,7 +121,7 @@ namespace Audio
 		, { EMusicalScale::Scale::LydianAugmented,
 			{{ Note::C1, Note::D1, Note::E1, Note::Gb1, Note::Ab1, Note::A1, Note::B1}, { Note::C1, Note::E1, Note::Ab1, Note::B1}}}
 		, { EMusicalScale::Scale::Augmented,
-			{{ Note::C1, Note::Eb1, Note::E1, Note::F1, Note::G1, Note::Ab1, Note::B1}, { Note::C1, Note::E1, Note::G1, Note::B1}}}
+			{{ Note::C1, Note::Eb1, Note::E1, Note::G1, Note::Ab1, Note::B1}, { Note::C1, Note::E1, Note::Ab1, Note::B1}}}
 		, { EMusicalScale::Scale::SixthModeOfHarmonicMinor,
 			{{ Note::C1, Note::Eb1, Note::E1, Note::Gb1, Note::G1, Note::A1, Note::B1}, { Note::C1, Note::E1, Note::G1, Note::B1}}}
 		, { EMusicalScale::Scale::Diminished_BeginWithHalfStep,
@@ -137,15 +137,15 @@ namespace Audio
 		, { EMusicalScale::Scale::Hindu,
 			{{ Note::C1, Note::D1, Note::E1, Note::F1, Note::G1, Note::Ab1, Note::Bb1}, { Note::C1, Note::E1, Note::G1, Note::Bb1}}}
 		, { EMusicalScale::Scale::WholeTone,
-			{{ Note::C1, Note::D1, Note::E1, Note::Gb1, Note::Ab1, Note::Bb1}, { Note::C1, Note::E1, Note::G1, Note::Bb1}}}
+			{{ Note::C1, Note::D1, Note::E1, Note::Gb1, Note::Ab1, Note::Bb1}, { Note::C1, Note::E1, Note::Gb1, Note::Bb1}}}
 		, { EMusicalScale::Scale::Chromatic,
 			{{ Note::C1, Note::Db1, Note::D1, Note::Eb1, Note::E1, Note::F1, Note::Gb1, Note::G1, Note::Ab1, Note::A1, Note::Bb1, Note::B1}, {Note::C1, Note::D1, Note::E1, Note::Gb1, Note::Ab1, Note::Bb1 /* Every other note */}}}
 		, { EMusicalScale::Scale::DiminishedWholeTone,
-			{{ Note::C1, Note::Db1, Note::Eb1, Note::E1, Note::Gb1, Note::G1, Note::A1, Note::Bb1}, { Note::C1, Note::E1, Note::G1, Note::Bb1}}}
+			{{ Note::C1, Note::Db1, Note::Eb1, Note::E1, Note::Gb1, Note::Ab1, Note::Bb1}, { Note::C1, Note::Eb1, Note::Gb1, Note::Bb1}}}
 		, { EMusicalScale::Scale::MinorPentatonic,
 			{{ Note::C1, Note::Eb1, Note::F1, Note::G1, Note::Bb1}, { Note::C1, Note::Eb1, Note::G1, Note::Bb1}}}
 		, { EMusicalScale::Scale::Bebop_Minor,
-			{{ Note::C1, Note::D1, Note::Eb1, Note::E1, Note::F1, Note::G1, Note::A1, Note::Bb1}, { Note::C1, Note::Eb1, Note::G1, Note::Bb1}}}
+			{{ Note::C1, Note::D1, Note::Eb1, Note::F1, Note::G1, Note::Ab1, Note::A1, Note::B1}, { Note::C1, Note::Eb1, Note::G1, Note::Bb1}}}
 		, { EMusicalScale::Scale::MelodicMinor,
 			{{ Note::C1, Note::D1, Note::Eb1, Note::F1, Note::G1, Note::A1, Note::B1}, { Note::C1, Note::Eb1, Note::G1, Note::B1}}}
 		, { EMusicalScale::Scale::Bebop_MinorNumber2,
@@ -153,12 +153,12 @@ namespace Audio
 		, { EMusicalScale::Scale::HarmonicMinor,
 			{{ Note::C1, Note::D1, Note::Eb1, Note::F1, Note::G1, Note::Ab1, Note::B1}, { Note::C1, Note::Eb1, Note::G1, Note::B1}}}
 		, { EMusicalScale::Scale::Diminished_BeginWithWholeStep,
-			{{ Note::C1, Note::D1, Note::Eb1, Note::F1, Note::Gb1, Note::Ab1, Note::A1, Note::B1}, { Note::C1, Note::Eb1, Note::Ab1, Note::B1}}}
+			{{ Note::C1, Note::D1, Note::Eb1, Note::F1, Note::Gb1, Note::Ab1, Note::A1, Note::B1}, { Note::C1, Note::Eb1, Note::Gb1, Note::B1}}}
 		, { EMusicalScale::Scale::Phrygian,
 			{{ Note::C1, Note::Db1, Note::Eb1, Note::F1, Note::G1, Note::Ab1, Note::Bb1}, { Note::C1, Note::Eb1, Note::G1, Note::Bb1}}}
 		, { EMusicalScale::Scale::NaturalMinor_Aeolian,
 			{{ Note::C1, Note::D1, Note::Eb1, Note::F1, Note::G1, Note::Ab1, Note::Bb1}, { Note::C1, Note::Eb1, Note::G1, Note::Bb1}}}
 		, { EMusicalScale::Scale::HalfDiminished_LocrianNumber2,
-			{{ Note::C1, Note::Db1, Note::Eb1, Note::F1, Note::Gb1, Note::Ab1, Note::Bb1}, { Note::C1, Note::Eb1, Note::Gb1, Note::Bb1}}}
+			{{ Note::C1, Note::D1, Note::Eb1, Note::F1, Note::Gb1, Note::Ab1, Note::Bb1}, { Note::C1, Note::Eb1, Note::Gb1, Note::Bb1}}}
 	};
 } // namespace Audio
