Audio dsp utilities for interpolating audio and multitap delaylines

#rnx

diff --git a/Engine/Source/Runtime/SignalProcessing/Private/FloatArrayMath.cpp b/Engine/Source/Runtime/SignalProcessing/Private/FloatArrayMath.cpp
--- a/Engine/Source/Runtime/SignalProcessing/Private/FloatArrayMath.cpp
+++ b/Engine/Source/Runtime/SignalProcessing/Private/FloatArrayMath.cpp
@@ -2347,6 +2347,89 @@ namespace Audio
 		}
 	}
 
+	void ArrayInterpolate(const float* InBuffer, float* OutBuffer, const int32 NumInSamples, const int32 NumOutSamples)
+	{
+		if (NumOutSamples <= 0 || NumInSamples <= 0)
+		{
+			return;
+		}
+
+		const float SampleStride = (float)NumInSamples / (float)NumOutSamples;
+
+		const int32 NumToSimd = NumOutSamples & MathIntrinsics::SimdMask;
+		const int32 NumNotToSimd = NumOutSamples & MathIntrinsics::NotSimdMask;
+
+		if(NumToSimd)
+		{
+			VectorRegister4Float Strides = VectorSet(
+				4.f * SampleStride,
+				4.f * SampleStride,
+				4.f * SampleStride,
+				4.f * SampleStride
+			);
+		
+			VectorRegister4Float Indeces = VectorSet(
+				0.f * SampleStride,
+				1.f * SampleStride,
+				2.f * SampleStride,
+				3.f * SampleStride
+			);
+			
+			for (int32 OutputIndex = 0; OutputIndex < NumToSimd; OutputIndex += AUDIO_NUM_FLOATS_PER_VECTOR_REGISTER)
+			{
+				alignas(16) int32 LeftIndecesRaw[4];
+				alignas(16) int32 RightIndecesRaw[4];
+			
+				VectorRegister4Float LeftIndeces = VectorFloor(Indeces);
+				VectorRegister4Float Fractions = VectorSubtract(Indeces, LeftIndeces);
+				VectorRegister4Float InvFractions = VectorSubtract(GlobalVectorConstants::FloatOne, Fractions);
+
+				VectorRegister4Int LeftIndecesInt = VectorFloatToInt(LeftIndeces);
+
+				// Lookup samples for interpolation
+				VectorIntStoreAligned(LeftIndecesInt, LeftIndecesRaw);
+				VectorIntStoreAligned(VectorIntAdd(LeftIndecesInt, GlobalVectorConstants::IntOne), RightIndecesRaw);
+
+				VectorRegister4Float LowerSamples = VectorSet(
+					InBuffer[LeftIndecesRaw[0]],
+					InBuffer[LeftIndecesRaw[1]],
+					InBuffer[LeftIndecesRaw[2]],
+					InBuffer[LeftIndecesRaw[3]]
+				);
+				VectorRegister4Float UpperSamples = VectorSet(
+					InBuffer[RightIndecesRaw[0]],
+					InBuffer[RightIndecesRaw[1]],
+					InBuffer[RightIndecesRaw[2]],
+					InBuffer[RightIndecesRaw[3]]
+				);
+			
+				VectorRegister4Float VOut = VectorMultiplyAdd(
+					LowerSamples,
+					Fractions,
+					VectorMultiply(UpperSamples, InvFractions));
+				VectorStore(VOut, &OutBuffer[OutputIndex]);
+
+				Indeces = VectorAdd(Indeces, Strides);
+			}
+		}
+
+		if(NumNotToSimd)
+		{
+			float SampleIndex = (float)(NumToSimd) * SampleStride;
+
+			for (int32 OutputIndex = NumToSimd; OutputIndex < NumOutSamples; OutputIndex++)
+            {
+            	const int32 LeftSample = FMath::FloorToInt32(SampleIndex);
+            	int32 RightSample = FMath::CeilToInt32(SampleIndex);
+            	
+            	const float Frac = SampleIndex - LeftSample;
+            	OutBuffer[OutputIndex] = (Frac * InBuffer[LeftSample]) + ((1.f - Frac) * InBuffer[RightSample]);
+            	
+            	SampleIndex += SampleStride;
+            }
+		}
+	}
+
 	FContiguousSparse2DKernelTransform::FContiguousSparse2DKernelTransform(const int32 NumInElements, const int32 NumOutElements)
 	:	NumIn(NumInElements)
 	,	NumOut(NumOutElements)
diff --git a/Engine/Source/Runtime/SignalProcessing/Private/InterpolatedMultiTapDelay.cpp b/Engine/Source/Runtime/SignalProcessing/Private/InterpolatedMultiTapDelay.cpp
new file mode 100644
--- /dev/null
+++ b/Engine/Source/Runtime/SignalProcessing/Private/InterpolatedMultiTapDelay.cpp
@@ -0,0 +1,148 @@
+// Copyright Epic Games, Inc. All Rights Reserved.
+
+
+#include "DSP/InterpolatedMultiTapDelay.h"
+
+#include "DSP/FloatArrayMath.h"
+#include "Math/VectorRegister.h"
+
+namespace Audio
+{
+	void FInterpolatedMultiTapDelay::Init(const int32 InBufferSizeSamples, const float InSampleRate)
+	{
+		WriteIndex = 0;
+		MsToSamples = InSampleRate / 1000.f;
+		DelayLine.Reset();
+		DelayLine.AddZeroed(AlignIndex(InBufferSizeSamples));
+	}
+
+	void FInterpolatedMultiTapDelay::Advance(const FAlignedFloatBuffer& InSamples)
+	{
+		check (InSamples.Num() % AUDIO_NUM_FLOATS_PER_VECTOR_REGISTER == 0);
+		const int32 InNumSamples = InSamples.Num();
+
+		if (InNumSamples <= 0)
+		{
+			return;
+		}
+
+		if (InNumSamples + WriteIndex > DelayLine.Num())
+		{
+			const int32 FirstHalfSize = DelayLine.Num() - WriteIndex;
+			const int32 SecondHalfSize = InNumSamples - FirstHalfSize;
+
+			if (FirstHalfSize > 0)
+			{
+				FMemory::Memcpy(&DelayLine[WriteIndex], InSamples.GetData(), FirstHalfSize * sizeof(float));
+			}
+			if (SecondHalfSize > 0)
+			{
+				FMemory::Memcpy(DelayLine.GetData(), &InSamples[FirstHalfSize], SecondHalfSize * sizeof(float));
+			}
+			WriteIndex = SecondHalfSize;
+		}
+		else
+		{
+			FMemory::Memcpy(&DelayLine[WriteIndex], InSamples.GetData(), InNumSamples * sizeof(float));
+			WriteIndex += InNumSamples;
+			WriteIndex = FMath::Wrap(WriteIndex, 0, DelayLine.Num());
+		}
+	}
+
+	void FInterpolatedMultiTapDelay::Read(const float StartDelayMSec, const float EndDelayMSec, FAlignedFloatBuffer& OutSamples)
+	{
+		const int32 OutputSamples = OutSamples.Num();
+		check (OutputSamples % AUDIO_NUM_FLOATS_PER_VECTOR_REGISTER == 0);
+
+		// likely to only run on the first frame, if not configured with enough memory
+		if (OutputSamples > DelayLine.Num())
+		{
+			DelayLine.SetNumZeroed(OutputSamples);
+		}
+
+		const int32 BufferSize = DelayLine.Num();
+		const int32 StartSample = AlignIndex(FMath::Wrap(WriteIndex - FMath::FloorToInt32(StartDelayMSec * MsToSamples), 0, BufferSize - 1));
+		const int32 EndSample = AlignIndex(FMath::Wrap(WriteIndex + OutputSamples - FMath::CeilToInt32(EndDelayMSec * MsToSamples), 0, BufferSize - 1));
+
+		float* OutputPtr = OutSamples.GetData();
+		
+		if (StartSample > EndSample)
+		{
+			const int32 DelaySamples = EndSample - (StartSample - BufferSize);
+
+			// block wraps
+			const int32 FirstBlockDelaySamples = BufferSize - StartSample;
+			const int32 SecondBlockDelaySamples = DelaySamples - FirstBlockDelaySamples;
+			const int32 FirstBlockOutputSamples = AlignIndex(((float)FirstBlockDelaySamples / (float)DelaySamples) * OutputSamples);
+			const int32 SecondBlockOuptutSamples = OutputSamples - FirstBlockOutputSamples;
+		
+			ReadBlockInternal(StartSample, FirstBlockDelaySamples, FirstBlockOutputSamples, OutputPtr);
+			ReadBlockInternal(0, SecondBlockDelaySamples, SecondBlockOuptutSamples, OutputPtr + FirstBlockOutputSamples);
+		}
+		else
+		{
+			ReadBlockInternal(StartSample, EndSample - StartSample, OutputSamples, OutputPtr);
+		}
+	}
+
+	void FInterpolatedMultiTapDelay::Reset()
+	{
+		const int32 NumElements = DelayLine.Num();
+		DelayLine.Reset(NumElements);
+		DelayLine.AddZeroed(NumElements);
+	}
+
+	bool FInterpolatedMultiTapDelay::IsInitialized() const
+	{
+		return DelayLine.GetAllocatedSize() > 0 && MsToSamples > 1.f;
+	}
+
+	void FInterpolatedMultiTapDelay::ReadBlockInternal(const int32 StartSample, const int32 SamplesToRead, const int32 NumOutputSamples, float* OutSamples)
+	{
+		if (NumOutputSamples == 0 || SamplesToRead == 0)
+		{
+			return;
+		}
+		
+		check(StartSample + SamplesToRead <= DelayLine.Num())
+		check(NumOutputSamples % AUDIO_NUM_FLOATS_PER_VECTOR_REGISTER == 0)
+
+		const float SampleStride = (float)SamplesToRead / (float)NumOutputSamples;
+
+		// protect against overflows while wrapping without branching in the main loop
+		int32 MainLoopOutputSamples = NumOutputSamples;
+		int32 MainLoopReadSamples = SamplesToRead;
+		
+		if (StartSample + SamplesToRead == DelayLine.Num())
+		{
+			MainLoopOutputSamples -= AUDIO_NUM_FLOATS_PER_VECTOR_REGISTER;
+			MainLoopReadSamples = FMath::RoundToInt32(MainLoopOutputSamples / SampleStride);
+		}
+
+		ArrayInterpolate(&DelayLine[StartSample], OutSamples, MainLoopReadSamples, MainLoopOutputSamples);
+
+		float SampleIndex = MainLoopReadSamples;
+		
+		// wrap in the loop only if close to wrapping: protects against rounding up past last index
+		for (int32 OutputIndex = MainLoopOutputSamples; OutputIndex < NumOutputSamples; OutputIndex++)
+		{
+			const int32 LeftSample = FMath::FloorToInt32(SampleIndex);
+			int32 RightSample = FMath::CeilToInt32(SampleIndex);
+			
+			if (RightSample >= DelayLine.Num())
+			{
+				RightSample -= DelayLine.Num();
+			}
+			
+			const float Frac = SampleIndex - LeftSample;
+			OutSamples[OutputIndex] = (Frac * DelayLine[LeftSample]) + ((1.f - Frac) * DelayLine[RightSample]);
+			
+			SampleIndex += SampleStride;
+		}
+	}
+
+	int32 FInterpolatedMultiTapDelay::AlignIndex(const int32 InIndex) const
+	{
+		return InIndex - (InIndex % AUDIO_NUM_FLOATS_PER_VECTOR_REGISTER);
+	}
+}
\ No newline at end of file
diff --git a/Engine/Source/Runtime/SignalProcessing/Public/DSP/FloatArrayMath.h b/Engine/Source/Runtime/SignalProcessing/Public/DSP/FloatArrayMath.h
--- a/Engine/Source/Runtime/SignalProcessing/Public/DSP/FloatArrayMath.h
+++ b/Engine/Source/Runtime/SignalProcessing/Public/DSP/FloatArrayMath.h
@@ -243,6 +243,9 @@ namespace Audio
 	/** Interleaves samples from an array of input buffers */
 	SIGNALPROCESSING_API void ArrayDeinterleave(const float* RESTRICT InBuffer, float** RESTRICT OutBuffers, const int32 InFrames, const int32 InChannels);
 
+	/** Interpolates a Mono audio buffer. */
+	SIGNALPROCESSING_API void ArrayInterpolate(const float* RESTRICT InBuffer, float* RESTRICT OutBuffer, const int32 NumInSamples, const int32 NumOutSamples);
+
 	/** FContiguousSparse2DKernelTransform
 	 *
 	 *  FContiguousSparse2DKernelTransform applies a matrix transformation to an input array. 
diff --git a/Engine/Source/Runtime/SignalProcessing/Public/DSP/InterpolatedMultiTapDelay.h b/Engine/Source/Runtime/SignalProcessing/Public/DSP/InterpolatedMultiTapDelay.h
new file mode 100644
--- /dev/null
+++ b/Engine/Source/Runtime/SignalProcessing/Public/DSP/InterpolatedMultiTapDelay.h
@@ -0,0 +1,36 @@
+// Copyright Epic Games, Inc. All Rights Reserved.
+
+#pragma once
+
+#include "AlignedBuffer.h"
+
+namespace Audio
+{
+	/**
+	 *  Delay line supporting multiple interpolated tap reads before advancing
+	 *  InBufferSizeSamples must be >= the size of the output buffer passed into Read
+	 *  If the delay length can be less than the output buffer size, Advance should be called before Read
+	 */
+	class FInterpolatedMultiTapDelay
+	{
+	public:
+		FInterpolatedMultiTapDelay() = default;
+		
+		SIGNALPROCESSING_API void Init(const int32 InBufferSizeSamples, const float InSampleRate);
+
+		SIGNALPROCESSING_API void Advance(const FAlignedFloatBuffer& InSamples);
+		SIGNALPROCESSING_API void Read(const float StartDelayMSec, const float EndDelayMSec, FAlignedFloatBuffer& OutSamples);
+		SIGNALPROCESSING_API void Reset();
+		SIGNALPROCESSING_API bool IsInitialized() const;
+
+	private:
+		void ReadBlockInternal(const int32 StartSample, const int32 SamplesToRead, const int32 NumOutputSamples, float* OutSamples);
+		FORCEINLINE int32 AlignIndex(const int32 InIndex) const;
+		
+		int32 WriteIndex = 0;
+		float MsToSamples = 1.f;
+
+		FAlignedFloatBuffer DelayLine;
+	};	
+}
+
