Fix unit test + possible buffer overrun in interpolatedmultitapdelay

#rnx

diff --git a/Engine/Source/Runtime/SignalProcessing/Private/InterpolatedMultiTapDelay.cpp b/Engine/Source/Runtime/SignalProcessing/Private/InterpolatedMultiTapDelay.cpp
--- a/Engine/Source/Runtime/SignalProcessing/Private/InterpolatedMultiTapDelay.cpp
+++ b/Engine/Source/Runtime/SignalProcessing/Private/InterpolatedMultiTapDelay.cpp
@@ -66,7 +66,7 @@ namespace Audio
 
 		float* OutputPtr = OutSamples.GetData();
 		
-		if (StartSample > EndSample)
+		if (StartSample >= EndSample)
 		{
 			const int32 DelaySamples = EndSample - (StartSample - BufferSize);
 
@@ -116,7 +116,7 @@ namespace Audio
 		if (StartSample + SamplesToRead == DelayLine.Num())
 		{
 			MainLoopOutputSamples -= AUDIO_NUM_FLOATS_PER_VECTOR_REGISTER;
-			MainLoopReadSamples = FMath::RoundToInt32(MainLoopOutputSamples / SampleStride);
+			MainLoopReadSamples -= FMath::RoundToInt32((float)AUDIO_NUM_FLOATS_PER_VECTOR_REGISTER * SampleStride);
 		}
 
 		ArrayInterpolate(&DelayLine[StartSample], OutSamples, MainLoopReadSamples, MainLoopOutputSamples);
