[SignalProcessing] Per-block process for FInlineEnvelopeFollower

diff --git a/Engine/Source/Runtime/SignalProcessing/Public/DSP/EnvelopeFollower.h b/Engine/Source/Runtime/SignalProcessing/Public/DSP/EnvelopeFollower.h
--- a/Engine/Source/Runtime/SignalProcessing/Public/DSP/EnvelopeFollower.h
+++ b/Engine/Source/Runtime/SignalProcessing/Public/DSP/EnvelopeFollower.h
@@ -4,6 +4,7 @@
 
 #include "Containers/Array.h"
 #include "DSP/AlignedBuffer.h"
+#include "DSP/FloatArrayMath.h"
 #include "DSP/Dsp.h"
 #include "HAL/Platform.h"
 
@@ -314,8 +315,8 @@ namespace Audio
 		float AnalysisWindowMsec=5.f;
 	};
 
-	/** FInlineEnvlepeFollower is useful for low samplerate use cases and where
-	 * samples are only available one at a time. This class is inlined becaused
+	/** FInlineEnvelopeFollower is useful for low sample rate use cases and where
+	 * samples are only available one at a time. This class is inlined because
 	 * there are situations where it is needed in a CPU intensive situations.
 	 */
 	class FInlineEnvelopeFollower : public FAttackRelease
@@ -385,6 +386,29 @@ namespace Audio
 			return Value;
 		}
 
+		FORCEINLINE void ProcessBuffer(const float* InSamples, const int32 InNumSamples, float* OutSamples)
+		{
+			check(InSamples != OutSamples);
+
+			NormalizeSamples(InSamples, InNumSamples, OutSamples);
+
+			for (int32 SampleIndex = 0; SampleIndex < InNumSamples; ++SampleIndex)
+			{
+				float NormSample = OutSamples[SampleIndex];
+				float Diff = Value - NormSample;
+				if (Diff <= 0.f)
+				{
+					Value = (GetAttackTimeSamples() * Diff) + NormSample;
+				}
+				else
+				{
+					Value = (GetReleaseTimeSamples() * Diff) + NormSample;
+				}
+				Value = Audio::UnderflowClamp(Value);
+				OutSamples[SampleIndex] = Value;
+			}
+		}
+
 		void Reset()
 		{
 			Value = 0.f;
@@ -431,6 +455,49 @@ namespace Audio
 			return InSample;
 		}
 
+		FORCEINLINE void NormalizeSamples(const float* InSamples, const int32 InNumSamples, float* OutSamples)
+		{
+			TArrayView<const float> In = TArrayView<const float>(InSamples, InNumSamples);
+			TArrayView<float> Out = TArrayView<float>(OutSamples, InNumSamples);
+			switch (Mode)
+			{
+			case EPeakMode::Peak:
+				ArrayAbs(In, Out);
+				return;
+
+			case EPeakMode::MeanSquared:
+			{
+				ArraySquare(In, Out);
+				for (int32 SampleIndex = 0; SampleIndex < InNumSamples; ++SampleIndex)
+				{
+					AnalysisValue = AnalysisFilterBeta * OutSamples[SampleIndex] + AnalysisFilterAlpha * AnalysisValue;
+					AnalysisValue = Audio::UnderflowClamp(AnalysisValue);
+					OutSamples[SampleIndex] = AnalysisValue;
+				}
+				return;
+			}
+
+			case EPeakMode::RootMeanSquared:
+
+			{
+				ArraySquare(In, Out);
+				for (int32 SampleIndex = 0; SampleIndex < InNumSamples; ++SampleIndex)
+				{
+					AnalysisValue = AnalysisFilterBeta * OutSamples[SampleIndex] + AnalysisFilterAlpha * AnalysisValue;
+					AnalysisValue = Audio::UnderflowClamp(AnalysisValue);
+					OutSamples[SampleIndex] = AnalysisValue;
+				}
+				ArraySqrtInPlace(Out);
+				return;
+			}
+
+			default:
+			{
+				checkNoEntry();
+			}
+			}
+		}
+
 		void InitMeanSquaredCoefficients(float InWinMsec, float InSampleRate)
 		{
 			if (ensure((InSampleRate > 0.f) && (InWinMsec > 0.f)))
