Add test scenario for ArrayMax in UintTestFloatArrayMath

diff --git a/Engine/Source/Runtime/SignalProcessing/Private/FloatArrayMath.cpp b/Engine/Source/Runtime/SignalProcessing/Private/FloatArrayMath.cpp
--- a/Engine/Source/Runtime/SignalProcessing/Private/FloatArrayMath.cpp
+++ b/Engine/Source/Runtime/SignalProcessing/Private/FloatArrayMath.cpp
@@ -731,6 +731,35 @@ namespace Audio
 		}
 	}
 	
+	void ArrayMax(const TArrayView<const float>& InView1, const TArrayView<const float>& InView2, const TArrayView<float>& OutView)
+	{
+		check(InView1.Num() == InView2.Num());
+		check(InView1.Num() == OutView.Num());
+
+		CSV_SCOPED_TIMING_STAT(Audio_Dsp, ArrayMax);
+		const int32 Num = InView1.Num();
+		const int32 NumToSimd = Num & MathIntrinsics::SimdMask;
+		const int32 NumNotToSimd = Num & MathIntrinsics::NotSimdMask;
+
+		if (NumToSimd)
+		{
+			for (int32 i = 0; i < NumToSimd; i += AUDIO_NUM_FLOATS_PER_VECTOR_REGISTER)
+			{
+				VectorRegister4Float Input1 = VectorLoad(&InView1[i]);
+				VectorRegister4Float Input2 = VectorLoad(&InView2[i]);
+				VectorStore(VectorMax(Input1, Input2), &OutView[i]);
+			}
+		}
+
+		if (NumNotToSimd)
+		{
+			for (int32 i = NumToSimd; i < Num; ++i)
+			{
+				OutView[i] = FMath::Max(InView1[i], InView2[i]);
+			}
+		}
+	}
+
 	float ArrayMaxAbsValue(const TArrayView<const float> InView)
 	{
 		CSV_SCOPED_TIMING_STAT(Audio_Dsp, ArrayMaxAbsValue);
@@ -2359,7 +2388,7 @@ namespace Audio
 		const int32 NumToSimd = NumOutSamples & MathIntrinsics::SimdMask;
 		const int32 NumNotToSimd = NumOutSamples & MathIntrinsics::NotSimdMask;
 
-		if(NumToSimd)
+		if (NumToSimd)
 		{
 			VectorRegister4Float Strides = VectorSet(
 				4.f * SampleStride,
@@ -2367,19 +2396,19 @@ namespace Audio
 				4.f * SampleStride,
 				4.f * SampleStride
 			);
-		
+
 			VectorRegister4Float Indeces = VectorSet(
 				0.f * SampleStride,
 				1.f * SampleStride,
 				2.f * SampleStride,
 				3.f * SampleStride
 			);
-			
+
 			for (int32 OutputIndex = 0; OutputIndex < NumToSimd; OutputIndex += AUDIO_NUM_FLOATS_PER_VECTOR_REGISTER)
 			{
 				alignas(16) int32 LeftIndecesRaw[4];
 				alignas(16) int32 RightIndecesRaw[4];
-			
+
 				VectorRegister4Float LeftIndeces = VectorFloor(Indeces);
 				VectorRegister4Float Fractions = VectorSubtract(Indeces, LeftIndeces);
 				VectorRegister4Float InvFractions = VectorSubtract(GlobalVectorConstants::FloatOne, Fractions);
@@ -2402,7 +2431,7 @@ namespace Audio
 					InBuffer[RightIndecesRaw[2]],
 					InBuffer[RightIndecesRaw[3]]
 				);
-			
+
 				VectorRegister4Float VOut = VectorMultiplyAdd(
 					LowerSamples,
 					Fractions,
@@ -2413,20 +2442,20 @@ namespace Audio
 			}
 		}
 
-		if(NumNotToSimd)
+		if (NumNotToSimd)
 		{
-			float SampleIndex = (float)(NumToSimd) * SampleStride;
+			float SampleIndex = (float)(NumToSimd)*SampleStride;
 
 			for (int32 OutputIndex = NumToSimd; OutputIndex < NumOutSamples; OutputIndex++)
-            {
-            	const int32 LeftSample = FMath::FloorToInt32(SampleIndex);
-            	int32 RightSample = FMath::CeilToInt32(SampleIndex);
-            	
-            	const float Frac = SampleIndex - LeftSample;
-            	OutBuffer[OutputIndex] = (Frac * InBuffer[LeftSample]) + ((1.f - Frac) * InBuffer[RightSample]);
-            	
-            	SampleIndex += SampleStride;
-            }
+			{
+				const int32 LeftSample = FMath::FloorToInt32(SampleIndex);
+				int32 RightSample = FMath::CeilToInt32(SampleIndex);
+
+				const float Frac = SampleIndex - LeftSample;
+				OutBuffer[OutputIndex] = (Frac * InBuffer[LeftSample]) + ((1.f - Frac) * InBuffer[RightSample]);
+
+				SampleIndex += SampleStride;
+			}
 		}
 	}
 
diff --git a/Engine/Source/Runtime/SignalProcessing/Public/DSP/FloatArrayMath.h b/Engine/Source/Runtime/SignalProcessing/Public/DSP/FloatArrayMath.h
--- a/Engine/Source/Runtime/SignalProcessing/Public/DSP/FloatArrayMath.h
+++ b/Engine/Source/Runtime/SignalProcessing/Public/DSP/FloatArrayMath.h
@@ -107,6 +107,10 @@ namespace Audio
 	 */ 
 	SIGNALPROCESSING_API void ArrayMinMaxNormalize(TArrayView<const float> InView, TArray<float>& OutArray);
 
+	/** Element-wise Max
+	 */
+	SIGNALPROCESSING_API void ArrayMax(const TArrayView<const float>& InView1, const TArrayView<const float>& InView2, const TArrayView<float>& OutView);
+
 	/** Returns the largest value of an array irrespective of sign (ex. {-3, 2, 1} would return 3).
 	 *  InView is a view of a float array to get the largest absolute value from.
 	 */
