[SignalProcessing] Non-interleaved processing for filters

diff --git a/Engine/Source/Runtime/SignalProcessing/Private/Filter.cpp b/Engine/Source/Runtime/SignalProcessing/Private/Filter.cpp
--- a/Engine/Source/Runtime/SignalProcessing/Private/Filter.cpp
+++ b/Engine/Source/Runtime/SignalProcessing/Private/Filter.cpp
@@ -73,6 +73,22 @@ namespace Audio
 		}
 	}
 
+	void FBiquadFilter::ProcessAudio(const float* const* InBuffers, const int32 InNumSamples, float* const* OutBuffers)
+	{
+		if (bEnabled)
+		{
+			for (int32 Channel = 0; Channel < NumChannels; ++Channel)
+			{
+				const float* Source = InBuffers[Channel];
+				float* Destination = OutBuffers[Channel];
+				for (int32 SampleIndex = 0; SampleIndex < InNumSamples; ++SampleIndex)
+				{
+					Destination[SampleIndex] = Biquad[Channel].ProcessAudio(Source[SampleIndex]);
+				}
+			}
+		}
+	}
+
 	void FBiquadFilter::Reset()
 	{
 		for (int32 Channel = 0; Channel < NumChannels; ++Channel)
@@ -543,6 +559,43 @@ namespace Audio
 		}
 	}
 
+	void FOnePoleFilter::ProcessAudio(const float* const* InBuffers, const int32 InNumSamples, float* const* OutBuffers)
+	{
+		if (FilterType == EFilter::HighPass)
+		{
+			for (int32 Channel = 0; Channel < NumChannels; ++Channel)
+			{
+				const float* Source = InBuffers[Channel];
+				float* Destination = OutBuffers[Channel];
+				for (int32 SampleIndex = 0; SampleIndex < InNumSamples; ++SampleIndex)
+				{
+					const float InputSample = Source[SampleIndex];
+					const float Vn = (InputSample - Z1[Channel]) * A0;
+					const float LPF = Vn + Z1[Channel];
+					Z1[Channel] = Vn + LPF;
+
+					Destination[SampleIndex] = InputSample - LPF;
+				}
+			}
+		}
+		else
+		{
+			for (int32 Channel = 0; Channel < NumChannels; ++Channel)
+			{
+				const float* Source = InBuffers[Channel];
+				float* Destination = OutBuffers[Channel];
+				for (int32 SampleIndex = 0; SampleIndex < InNumSamples; ++SampleIndex)
+				{
+					const float Vn = (Source[SampleIndex] - Z1[Channel]) * A0;
+					const float LPF = Vn + Z1[Channel];
+					Z1[Channel] = Vn + LPF;
+
+					Destination[SampleIndex] = LPF;
+				}
+			}
+		}
+	}
+
 	FStateVariableFilter::FStateVariableFilter()
 		: InputScale(1.0f)
 		, A0(1.0f)
@@ -630,6 +683,42 @@ namespace Audio
 		}
 	}
 
+	void FStateVariableFilter::ProcessAudio(const float* const* InBuffers, const int32 InNumSamples, float* const* OutBuffers)
+	{
+		float FilterResults[4];
+		enum ResultSourceType { kLPF = 0, kHPF = 1, kBPF = 2, kBSF = 3 };
+		int32 OutputSource = 0;
+		switch (FilterType)
+		{
+		default:
+		case EFilter::LowPass:  OutputSource = ResultSourceType::kLPF; break;
+		case EFilter::HighPass: OutputSource = ResultSourceType::kHPF; break;
+		case EFilter::BandPass: OutputSource = ResultSourceType::kBPF; break;
+		case EFilter::BandStop: OutputSource = ResultSourceType::kBSF; break;
+		}
+
+		for (int32 Channel = 0; Channel < NumChannels; ++Channel)
+		{
+			const float* Source = InBuffers[Channel];
+			float* Destination = OutBuffers[Channel];
+			FFilterState& ChannelFilterState = FilterState[Channel];
+			for (int32 SampleIndex = 0; SampleIndex < InNumSamples; ++SampleIndex)
+			{
+				FilterResults[ResultSourceType::kHPF] = InputScale * (Source[SampleIndex] - Feedback * ChannelFilterState.Z1_1 - ChannelFilterState.Z1_2);
+				FilterResults[ResultSourceType::kBPF] = Audio::FastTanh(A0 * FilterResults[ResultSourceType::kHPF] + ChannelFilterState.Z1_1);
+
+				FilterResults[kLPF] = A0 * FilterResults[kBPF] + ChannelFilterState.Z1_2;
+				const float Dampening = 0.5f / Q;
+				FilterResults[kBSF] = BandStopParam * FilterResults[kHPF] + (1.0f - BandStopParam) * FilterResults[kLPF];
+
+				ChannelFilterState.Z1_1 = A0 * FilterResults[kHPF] + FilterResults[kBPF];
+				ChannelFilterState.Z1_2 = A0 * FilterResults[kBPF] + FilterResults[kLPF];
+
+				Destination[SampleIndex] = FilterResults[OutputSource];
+			}
+		}
+
+	}
 
 	void FStateVariableFilter::ProcessAudio(const float* InSamples, const int32 InNumSamples
 		, float* LpfOutput, float* HpfOutput, float* BpfOutput, float* BsfOutput)
@@ -656,6 +745,36 @@ namespace Audio
 		}
 	}
 
+	void FStateVariableFilter::ProcessAudio(const float* const* InBuffers, const int32 InNumSamples, 
+		float* const* LpfOutBuffers, float* const* HpfOutBuffers, float* const* BpfOutBuffers, float* const* BsfOutBuffers)
+	{
+		for (int32 Channel = 0; Channel < NumChannels; ++Channel)
+		{
+			const float* Source = InBuffers[Channel];
+			float* LpfDestination = LpfOutBuffers[Channel];
+			float* HpfDestination = HpfOutBuffers[Channel];
+			float* BpfDestination = BpfOutBuffers[Channel];
+			float* BsfDestination = BsfOutBuffers[Channel];
+			FFilterState& ChannelFilterState = FilterState[Channel];
+			for (int32 SampleIndex = 0; SampleIndex < InNumSamples; ++SampleIndex)
+			{
+				const float HPF = InputScale * (Source[SampleIndex] - Feedback * ChannelFilterState.Z1_1 - ChannelFilterState.Z1_2);
+				float BPF = Audio::FastTanh(A0 * HPF + ChannelFilterState.Z1_1);
+
+				const float LPF = A0 * BPF + ChannelFilterState.Z1_2;
+				const float Dampening = 0.5f / Q;
+				const float BSF = BandStopParam * HPF + (1.0f - BandStopParam) * LPF;
+
+				ChannelFilterState.Z1_1 = A0 * HPF + BPF;
+				ChannelFilterState.Z1_2 = A0 * BPF + LPF;
+
+				LpfDestination[SampleIndex] = LPF;
+				HpfDestination[SampleIndex] = HPF;
+				BpfDestination[SampleIndex] = BPF;
+				BsfDestination[SampleIndex] = BSF;
+			}
+		}
+	}
 	FLadderFilter::FLadderFilter()
 		: K(0.0f)
 		, Gamma(0.0f)
@@ -802,4 +921,46 @@ namespace Audio
 		}
 	}
 
+	void FLadderFilter::ProcessAudio(const float* const* InBuffers, const int32 InNumSamples, float* const* OutBuffers)
+	{
+		float U[MaxFilterChannels];
+
+		float OutputFilter0[MaxFilterChannels];
+		float OutputFilter1[MaxFilterChannels];
+		float OutputFilter2[MaxFilterChannels];
+		float OutputFilter3[MaxFilterChannels];
+
+		for (int32 SampleIndex = 0; SampleIndex < InNumSamples; ++SampleIndex)
+		{
+			for (int32 Channel = 0; Channel < NumChannels; ++Channel)
+			{
+				float Sigma = 0.0f;
+				for (int32 i = 0; i < 4; ++i)
+				{
+					Sigma += OnePoleFilters[i].GetState(Channel) * Beta[i];
+				}
+
+				float InSample = InBuffers[Channel][SampleIndex];
+				InSample *= 1.0f + PassBandGainCompensation * K;
+
+				// Compute input into first LPF
+				U[Channel] = FMath::Min(Audio::FastTanh((InSample - K * Sigma) * Alpha), 1.0f);
+			}
+
+			OnePoleFilters[0].ProcessAudioFrame(U, OutputFilter0);
+			OnePoleFilters[1].ProcessAudioFrame(OutputFilter0, OutputFilter1);
+			OnePoleFilters[2].ProcessAudioFrame(OutputFilter1, OutputFilter2);
+			OnePoleFilters[3].ProcessAudioFrame(OutputFilter2, OutputFilter3);
+
+			for (int32 Channel = 0; Channel < NumChannels; ++Channel)
+			{
+				// Feed U into first filter, then cascade down
+				OutBuffers[Channel][SampleIndex] = Factors[0] * U[Channel];
+				OutBuffers[Channel][SampleIndex] += Factors[1] * OutputFilter0[Channel];
+				OutBuffers[Channel][SampleIndex] += Factors[2] * OutputFilter1[Channel];
+				OutBuffers[Channel][SampleIndex] += Factors[3] * OutputFilter2[Channel];
+				OutBuffers[Channel][SampleIndex] += Factors[4] * OutputFilter3[Channel];
+			}
+		}
+	}
 }
diff --git a/Engine/Source/Runtime/SignalProcessing/Public/DSP/Filter.h b/Engine/Source/Runtime/SignalProcessing/Public/DSP/Filter.h
--- a/Engine/Source/Runtime/SignalProcessing/Public/DSP/Filter.h
+++ b/Engine/Source/Runtime/SignalProcessing/Public/DSP/Filter.h
@@ -49,9 +49,12 @@ namespace Audio
 		// Processes a single frame of audio
 		SIGNALPROCESSING_API void ProcessAudioFrame(const float* InFrame, float* OutFrame);
 
-		// Process a buffer of audio
+		// Process a mono buffer or an interleaved buffer of multichannel audio
 		SIGNALPROCESSING_API void ProcessAudio(const float* InBuffer, const int32 InNumSamples, float* OutBuffer);
 
+		// Process a non-interleaved buffer of multichannel audio
+		SIGNALPROCESSING_API void ProcessAudio(const float* const* InBuffers, const int32 InNumSamples, float* const* OutBuffers);
+
 		// Sets all filter parameters with one function
 		SIGNALPROCESSING_API void SetParams(const EBiquadFilter::Type InFilterType, const float InCutoffFrequency, const float InBandwidth, const float InGainDB);
 
@@ -164,9 +167,12 @@ namespace Audio
 			ProcessAudio(InFrame, NumChannels, OutFrame);
 		}
 
-		// Process an audio buffer.
+		// Process a mono buffer or an interleaved buffer of multichannel audio
 		virtual void ProcessAudio(const float* InBuffer, const int32 InNumSamples, float* OutBuffer) = 0;
 
+		// Process a non-interleaved buffer of multichannel audio
+		virtual void ProcessAudio(const float* const* InBuffers, const int32 InNumSamples, float* const* OutBuffers) = 0;
+
 		// Filter patch destinations
 		FPatchDestination GetModDestCutoffFrequency() const { return ModCutoffFrequencyDest; }
 		FPatchDestination GetModDestQ() const { return ModQDest; }
@@ -238,7 +244,10 @@ namespace Audio
 		SIGNALPROCESSING_API virtual void Reset() override;
 		SIGNALPROCESSING_API virtual void Update() override;
 		SIGNALPROCESSING_API virtual void ProcessAudioFrame(const float* InFrame, float* OutFrame) override;
+		// Process a mono buffer or an interleaved buffer of multichannel audio
 		SIGNALPROCESSING_API virtual void ProcessAudio(const float* InSamples, const int32 InNumSamples, float* OutSamples) override;
+		// Process a non-interleaved buffer of multichannel audio
+		SIGNALPROCESSING_API virtual void ProcessAudio(const float* const* InBuffers, const int32 InNumSamples, float* const* OutBuffers) override;
 
 		void SetCoefficient(const float InCoefficient) { A0 = InCoefficient; }
 		float GetCoefficient() const { return A0; }
@@ -260,10 +269,17 @@ namespace Audio
 		SIGNALPROCESSING_API virtual void SetBandStopControl(const float InBandStop) override;
 		SIGNALPROCESSING_API virtual void Reset() override;
 		SIGNALPROCESSING_API virtual void Update() override;
+		// Process a mono buffer or an interleaved buffer of multichannel audio
 		SIGNALPROCESSING_API virtual void ProcessAudio(const float* InSamples, const int32 InNumSamples, float* OutSamples) override;
+		// Process a non-interleaved buffer of multichannel audio
+		SIGNALPROCESSING_API virtual void ProcessAudio(const float* const* InBuffers, const int32 InNumSamples, float* const* OutBuffers) override;
 
-		SIGNALPROCESSING_API void ProcessAudio(const float* InSamples, const int32 InNumSamples
-			, float* LpfOutput, float* HpfOutput, float* BpfOutput, float* BsfOutput);
+		// Process a mono buffer or an interleaved buffer of multichannel audio
+		SIGNALPROCESSING_API void ProcessAudio(const float* InSamples, const int32 InNumSamples,
+			float* LpfOutput, float* HpfOutput, float* BpfOutput, float* BsfOutput);
+		// Process a non-interleaved buffer of multichannel audio
+		SIGNALPROCESSING_API void ProcessAudio(const float* const* InBuffers, const int32 InNumSamples,
+			float* const* LpfOutBuffers, float* const* HpfOutBuffers, float* const* BpfOutBuffers, float* const* BsfOutBuffers);
 
 	protected:
 		float InputScale;
@@ -297,8 +313,12 @@ namespace Audio
 		SIGNALPROCESSING_API virtual void SetQ(const float InQ) override;
 		SIGNALPROCESSING_API virtual void SetPassBandGainCompensation(const float InPassBandGainCompensation) override;
 
+		// Process a mono buffer or an interleaved buffer of multichannel audio
 		SIGNALPROCESSING_API virtual void ProcessAudio(const float* InSamples, const int32 InNumSamples, float* OutSamples) override;
 
+		// Process a non-interleaved buffer of multichannel audio
+		SIGNALPROCESSING_API virtual void ProcessAudio(const float* const* InBuffers, const int32 InNumSamples, float* const* OutBuffers) override;
+
 
 	protected:
 
