[SignalProcessing] Per-block process for circular buffer delay

diff --git a/Engine/Source/Runtime/SignalProcessing/Private/Delay.cpp b/Engine/Source/Runtime/SignalProcessing/Private/Delay.cpp
--- a/Engine/Source/Runtime/SignalProcessing/Private/Delay.cpp
+++ b/Engine/Source/Runtime/SignalProcessing/Private/Delay.cpp
@@ -173,6 +173,21 @@ namespace Audio
 		return OutputAttenuation * Yn;
 	}
 
+	void FDelay::ProcessAudioBuffer(const float* InAudio, int32 InNumSamples, float* OutAudio)
+	{
+		// Note: There is probably some optization that could be done here with 
+		// memcpys or someting, but for now we will do the simple version. Obviously
+		// it could get very complicated when the delay buffer is smaller than the
+		// number of samples being requested. 
+		for (int32 SampleIndex = 0; SampleIndex < InNumSamples; ++SampleIndex)
+		{
+			Update();
+			const float Yn = DelayInSamples == 0 ? InAudio[SampleIndex] : Read();
+			WriteDelayAndInc(InAudio[SampleIndex]);
+			OutAudio[SampleIndex] = OutputAttenuation * Yn;
+		}
+	}
+
 	void FDelay::Update(bool bForce)
 	{
 		if (!EaseDelayMsec.IsDone() || bForce)
diff --git a/Engine/Source/Runtime/SignalProcessing/Public/DSP/Delay.h b/Engine/Source/Runtime/SignalProcessing/Public/DSP/Delay.h
--- a/Engine/Source/Runtime/SignalProcessing/Public/DSP/Delay.h
+++ b/Engine/Source/Runtime/SignalProcessing/Public/DSP/Delay.h
@@ -55,6 +55,8 @@ namespace Audio
 		// Process audio in the delay line, return the delayed value
 		SIGNALPROCESSING_API virtual float ProcessAudioSample(const float InAudio);
 
+		SIGNALPROCESSING_API virtual void ProcessAudioBuffer(const float* InAudio, int32 InNumSamples, float* OutAudio);
+
 	protected:
 
 		// Updates delay line based on any recent changes to settings
