[SignalProcessing] Vectorized process method for DynamicsProcessor

diff --git a/Engine/Source/Runtime/SignalProcessing/Private/DynamicsProcessor.cpp b/Engine/Source/Runtime/SignalProcessing/Private/DynamicsProcessor.cpp
--- a/Engine/Source/Runtime/SignalProcessing/Private/DynamicsProcessor.cpp
+++ b/Engine/Source/Runtime/SignalProcessing/Private/DynamicsProcessor.cpp
@@ -1,12 +1,14 @@
 // Copyright Epic Games, Inc. All Rights Reserved.
 
 #include "DSP/DynamicsProcessor.h"
+#include "DSP/FloatArrayMath.h"
 #include "SignalProcessingModule.h"
 
 namespace Audio
 {
 	FDynamicsProcessor::FDynamicsProcessor()
 		: ProcessingMode(EDynamicsProcessingMode::Compressor)
+		, SlopeFactor(0.0f)
 		, EnvelopeFollowerPeakMode(EPeakMode::Peak)
 		, LookaheedDelayMsec(10.0f)
 		, AttackTimeMsec(20.0f)
@@ -24,7 +26,8 @@ namespace Audio
 		, bKeyLowshelfEnabled(false)
 	{
 		// The knee will have 2 points
-		KneePoints.Init(FVector2D(), 2);
+		KneePoints.Init(FKneePoint(), 2);
+		CalculateSlope();
 	}
 
 	FDynamicsProcessor::~FDynamicsProcessor()
@@ -104,17 +107,20 @@ namespace Audio
 	void FDynamicsProcessor::SetThreshold(const float InThresholdDb)
 	{
 		ThresholdDb = InThresholdDb;
+		CalculateKnee();
 	}
 
 	void FDynamicsProcessor::SetRatio(const float InCompressionRatio)
 	{
 		// Don't let the compression ratio be 0.0!
 		Ratio = FMath::Max(InCompressionRatio, SMALL_NUMBER);
+		CalculateSlope();
 	}
 
 	void FDynamicsProcessor::SetKneeBandwidth(const float InKneeBandwidthDb)
 	{
 		HalfKneeBandwidthDb = 0.5f * InKneeBandwidthDb;
+		CalculateKnee();
 	}
 
 	void FDynamicsProcessor::SetInputGain(const float InInputGainDb)
@@ -243,6 +249,7 @@ namespace Audio
 	void FDynamicsProcessor::SetProcessingMode(const EDynamicsProcessingMode::Type InProcessingMode)
 	{
 		ProcessingMode = InProcessingMode;
+		CalculateSlope();
 	}
 
 	void FDynamicsProcessor::ProcessAudioFrame(const float* InFrame, float* OutFrame, const float* InKeyFrame)
@@ -324,7 +331,6 @@ namespace Audio
 		const int32 NumChannels = GetNumChannels();
 		const int32 KeyNumChannels = GetKeyNumChannels();
 
-
 		if (InKeyBuffer)
 		{
 			int32 KeySampleIndex = 0;
@@ -345,6 +351,140 @@ namespace Audio
 		}
 	}
 
+	void FDynamicsProcessor::ProcessAudio(const float* const* const InBuffers, const int32 InNumSamples, float* const* OutBuffers, const float* const* const InKeyBuffers, float* const* OutEnvelopes)
+	{
+		// NOTE: This is a useful test, but it isn't comprehensive. It can't be sure the passed 
+		// in arrays of pointers have a matching number of pointers to the channel counts. It
+		// will have to do for now.
+		check(nullptr != InBuffers && nullptr != InBuffers[0]);
+		check(nullptr != OutBuffers && nullptr != OutBuffers[0]);
+		check(nullptr != OutEnvelopes && nullptr != OutEnvelopes[0]);
+
+		const int32 KeyNumChannels = GetKeyNumChannels();
+		const int32 NumChannels = GetNumChannels();
+
+		// we need this test because we are going to use the output buffer as our scratch pad for key processing!
+		check(KeyNumChannels <= NumChannels);
+
+		const float* const* KeySources = InKeyBuffers ? InKeyBuffers : InBuffers;
+		float* const* KeyWorkBuffers = OutBuffers;
+
+		// Process Key
+		if (bKeyLowshelfEnabled)
+		{
+			InputLowshelfFilter.ProcessAudio(KeySources, InNumSamples, KeyWorkBuffers);
+			KeySources = KeyWorkBuffers;
+		}
+
+		if (bKeyHighshelfEnabled)
+		{
+			InputHighshelfFilter.ProcessAudio(KeySources, InNumSamples, KeyWorkBuffers);
+			KeySources = KeyWorkBuffers;
+		}
+
+		const bool bKeyIsInput = !InKeyBuffers || InKeyBuffers == InBuffers;
+		float DetectorGain = InputGain;
+
+		// Apply key gain only if detector is key (not input)
+		if (!bKeyIsInput)
+		{
+			DetectorGain *= KeyGain;
+		}
+
+		if (bKeyAuditionEnabled)
+		{
+			if (KeySources != KeyWorkBuffers)
+			{
+				// This means we have not done either filter above, so we have not  
+				// "moved" the key input to the output buffer. We have to do it here. 
+				for (int32 Channel = 0; Channel < KeyNumChannels; ++Channel)
+				{
+					FMemory::Memcpy(OutBuffers[Channel], KeySources[Channel], sizeof(float) * InNumSamples);
+				}
+			}
+			// we just have to zero out any channels that are not in the key signal
+			for (int32 Channel = KeyNumChannels; Channel < NumChannels; ++Channel)
+			{
+				FMemory::Memset(OutBuffers[Channel], 0, sizeof(float) * InNumSamples);
+			}
+
+			// I WOULD zero or 1 the "OutEnvelope", but that is not what the original does!
+			return;
+		}
+
+		for (int32 Channel = 0; Channel < KeyNumChannels; ++Channel)
+		{
+			// apply detector gain to key
+			Audio::ArrayMultiplyByConstant(TArrayView<const float>(KeySources[Channel], InNumSamples), DetectorGain, TArrayView<float>(KeyWorkBuffers[Channel], InNumSamples));
+			// EnvelopeFollow key
+			EnvFollower[Channel].ProcessBuffer(KeyWorkBuffers[Channel], InNumSamples, OutEnvelopes[Channel]);
+		}
+
+		int32 NumGainChannels = 1;
+		switch (LinkMode)
+		{
+		case EDynamicsProcessorChannelLinkMode::Average:
+		{
+			// average all key channels
+			TArrayView<float> Out = TArrayView<float>(OutEnvelopes[0], InNumSamples);
+			for (int32 KeyChannel = 1; KeyChannel < KeyNumChannels; ++KeyChannel)
+			{
+				ArraySum(TArrayView<const float>(OutEnvelopes[KeyChannel], InNumSamples), Out, Out);
+			}
+			ArrayMultiplyByConstant(Out, 1.0f / (float)KeyNumChannels, Out);
+			// convert magnitude to db
+			ArrayMagnitudeToDecibelInPlace(Out, -96.0f);
+			// compute 1 gain and use for all channels ...	const float ComputedGain = ComputeGain(DetectorOutLinkedDb);
+			ComputeGains(OutEnvelopes[0], InNumSamples);
+			NumGainChannels = 1;
+		}
+		break;
+
+		case EDynamicsProcessorChannelLinkMode::Peak:
+		{
+			// detect peak across all key channels
+			TArrayView<float> Out = TArrayView<float>(OutEnvelopes[0], InNumSamples);
+			for (int32 KeyChannel = 1; KeyChannel < KeyNumChannels; ++KeyChannel)
+			{
+				ArrayMax(TArrayView<const float>(OutEnvelopes[KeyChannel], InNumSamples), Out, Out);
+			}
+			// convert magnitude to db
+			ArrayMagnitudeToDecibelInPlace(Out, -96.0f);
+			// compute 1 gain and use for all channels ...	const float ComputedGain = ComputeGain(DetectorOutLinkedDb);
+			ComputeGains(OutEnvelopes[0], InNumSamples);
+			NumGainChannels = 1;
+		}
+		break;
+
+		case EDynamicsProcessorChannelLinkMode::Disabled:
+		default:
+		{
+			for (int32 Channel = 0; Channel < KeyNumChannels; ++Channel)
+			{
+				// convert magnitude to db
+				ArrayMagnitudeToDecibelInPlace(TArrayView<float>(OutEnvelopes[Channel],InNumSamples), -96.0f);
+				// compute 1 gain and use for all channels ...	const float ComputedGain = ComputeGain(DetectorOutLinkedDb);
+				ComputeGains(OutEnvelopes[Channel], InNumSamples);
+			}
+			NumGainChannels = KeyNumChannels;
+		}
+		break;
+		}
+
+		const float OutAndInGain = OutputGain * InputGain;
+		for (int32 Channel = 0; Channel < NumChannels; ++Channel)
+		{
+			int32 GainChannel = Channel % NumGainChannels;
+
+			// copy input to look ahead delay write position
+			LookaheadDelay[Channel].ProcessAudioBuffer(InBuffers[Channel], InNumSamples, OutBuffers[Channel]);
+			
+			// apply compression gain & output gain & input gain to look ahead delay read position and save to output
+			ArrayMultiplyInPlace(TArrayView<const float>(OutEnvelopes[GainChannel], InNumSamples), TArrayView<float>(OutBuffers[Channel], InNumSamples));
+			ArrayMultiplyByConstant(TArrayView<const float>(OutBuffers[Channel], InNumSamples), OutAndInGain, TArrayView<float>(OutBuffers[Channel], InNumSamples));
+		}
+	}
+
 	bool FDynamicsProcessor::ProcessKeyFrame(const float* InKeyFrame, float* OutFrame, bool bKeyIsInput)
 	{
 		// Get detector outputs
@@ -459,64 +599,120 @@ namespace Audio
 
 	float FDynamicsProcessor::ComputeGain(const float InEnvFollowerDb)
 	{
-		float SlopeFactor = 0.0f;
+		// If we are in the range of compression
+		float SlopeThisSample = SlopeFactor;
+		if (IsInProcessingThreshold(InEnvFollowerDb))
+		{
+			// The knee calculation adjusts the slope to use via lagrangian interpolation through the slope
+			float Lagrangian = (InEnvFollowerDb - KneePoints[1].X) / Denominator0Minus1;
+			SlopeThisSample = Lagrangian * KneePoints[0].Y;
+			Lagrangian = (InEnvFollowerDb - KneePoints[0].X) / Denominator1Minus0;
+			SlopeThisSample += Lagrangian * KneePoints[1].Y;
+
+		}
+
+		float OutputGainDb = SlopeThisSample * (ThresholdDb - InEnvFollowerDb);
+
+		if (ProcessingMode == EDynamicsProcessingMode::UpwardsCompressor)
+		{
+			// if left unchecked Upwards compression will try to apply infinite gain
+			OutputGainDb = FMath::Clamp(OutputGainDb, 0.f, UpwardsCompressionMaxGain);
+		}
+		else
+		{
+			OutputGainDb = FMath::Min(0.f, OutputGainDb);
+		}
+
+		return ConvertToLinear(OutputGainDb);
+	}
+
+	
+	void FDynamicsProcessor::ComputeGains(float* InEnvFollowerDbOutGain, const int32 InNumSamples)
+	{
+		for (int32 SampleIndex = 0; SampleIndex < InNumSamples; ++SampleIndex)
+		{
+			float InEnvFollowerDb = InEnvFollowerDbOutGain[SampleIndex];
+
+			float SlopeThisSample = SlopeFactor;
+			if (IsInProcessingThreshold(InEnvFollowerDb))
+			{
+				float Lagrangian = (InEnvFollowerDb - KneePoints[1].X) / Denominator0Minus1;
+				SlopeThisSample = Lagrangian * KneePoints[0].Y;
+				Lagrangian = (InEnvFollowerDb - KneePoints[0].X) / Denominator1Minus0;
+				SlopeThisSample += Lagrangian * KneePoints[1].Y;
+			};
+
+			float OutputGainDb = SlopeThisSample * (ThresholdDb - InEnvFollowerDb);
 
+			if (ProcessingMode == EDynamicsProcessingMode::UpwardsCompressor)
+			{
+				// if left unchecked Upwards compression will try to apply infinite gain
+				OutputGainDb = FMath::Clamp(OutputGainDb, 0.f, UpwardsCompressionMaxGain);
+				InEnvFollowerDbOutGain[SampleIndex] = ConvertToLinear(OutputGainDb);
+			}
+			else
+			{
+				InEnvFollowerDbOutGain[SampleIndex] = OutputGainDb > -UE_SMALL_NUMBER ? 1.0f : ConvertToLinear(OutputGainDb);
+			}
+		}
+	}
+
+	void FDynamicsProcessor::CalculateSlope()
+	{
+		SlopeFactor = 0.0f;
 		// Depending on the mode, we define the "slope". 
 		switch (ProcessingMode)
 		{
-			default:
+		default:
 
 			// Compressors smoothly reduce the gain as the gain gets louder
 			// CompressionRatio -> Inifinity is a limiter
 			// Upwards compression applies gain when below a threshold, but uses the same slope
-			case EDynamicsProcessingMode::UpwardsCompressor:
-			case EDynamicsProcessingMode::Compressor:
-				SlopeFactor = 1.0f - 1.0f / Ratio;
-				break;
+		case EDynamicsProcessingMode::UpwardsCompressor:
+		case EDynamicsProcessingMode::Compressor:
+			SlopeFactor = 1.0f - 1.0f / Ratio;
+			break;
 
 			// Limiters do nothing until it hits the threshold then clamps the output hard
-			case EDynamicsProcessingMode::Limiter:
-				SlopeFactor = 1.0f;
-				break;
+		case EDynamicsProcessingMode::Limiter:
+			SlopeFactor = 1.0f;
+			break;
 
 			// Expanders smoothly increase the gain as the gain gets louder
 			// CompressionRatio -> Infinity is a gate
-			case EDynamicsProcessingMode::Expander:
-				SlopeFactor = 1.0f / Ratio - 1.0f;
-				break;
+		case EDynamicsProcessingMode::Expander:
+			SlopeFactor = 1.0f / Ratio - 1.0f;
+			break;
 
 			// Gates are opposite of limiter. They stop sound (stop gain) until the threshold is hit
-			case EDynamicsProcessingMode::Gate:
-				SlopeFactor = -1.0f;
-				break;
+		case EDynamicsProcessingMode::Gate:
+			SlopeFactor = -1.0f;
+			break;
 		}
+		CalculateKnee();
+	}
 
-		// If we are in the range of compression
-		if (IsInProcessingThreshold(InEnvFollowerDb))
-		{
-			// Setup the knee for interpolation. Don't allow the top knee point to exceed 0.0
-			KneePoints[0].X = ThresholdDb - HalfKneeBandwidthDb;
-			KneePoints[1].X = FMath::Min(ThresholdDb + HalfKneeBandwidthDb, 0.0f);
-
-			KneePoints[0].Y = 0.0f;
-			KneePoints[1].Y = SlopeFactor;
-
-			// The knee calculation adjusts the slope to use via lagrangian interpolation through the slope
-			SlopeFactor = LagrangianInterpolation(KneePoints, InEnvFollowerDb);
-		}
+	void FDynamicsProcessor::CalculateKnee()
+	{
+		// Setup the knee for interpolation. Don't allow the top knee point to exceed 0.0
+		KneePoints[0].X = ThresholdDb - HalfKneeBandwidthDb;
+		KneePoints[1].X = FMath::Min(ThresholdDb + HalfKneeBandwidthDb, 0.0f);
 
-		float OutputGainDb = SlopeFactor * (ThresholdDb - InEnvFollowerDb);
+		KneePoints[0].Y = 0.0f;
+		KneePoints[1].Y = SlopeFactor;
 
-		if (ProcessingMode == EDynamicsProcessingMode::UpwardsCompressor)
+		// These next few calculations let us optimize out the call to 
+		// the LagrangeInterpolation that used to be in there. We precalculate 
+		// some coefficients that remain the same unless the KneePoints change.
+		Denominator0Minus1 = KneePoints[0].X - KneePoints[1].X;
+		if (FMath::Abs(Denominator0Minus1) < UE_SMALL_NUMBER)
 		{
-			// if left unchecked Upwards compression will try to apply infinite gain
-			OutputGainDb = FMath::Clamp(OutputGainDb, 0.f, UpwardsCompressionMaxGain);
+			Denominator0Minus1 = UE_SMALL_NUMBER;
 		}
-		else
+		Denominator1Minus0 = KneePoints[1].X - KneePoints[0].X;
+		if (FMath::Abs(Denominator1Minus0) < UE_SMALL_NUMBER)
 		{
-			OutputGainDb = FMath::Min(0.f, OutputGainDb);
+			Denominator1Minus0 = UE_SMALL_NUMBER;
 		}
-
-		return ConvertToLinear(OutputGainDb);
 	}
 }
diff --git a/Engine/Source/Runtime/SignalProcessing/Public/DSP/DynamicsProcessor.h b/Engine/Source/Runtime/SignalProcessing/Public/DSP/DynamicsProcessor.h
--- a/Engine/Source/Runtime/SignalProcessing/Public/DSP/DynamicsProcessor.h
+++ b/Engine/Source/Runtime/SignalProcessing/Public/DSP/DynamicsProcessor.h
@@ -73,10 +73,12 @@ namespace Audio
 		SIGNALPROCESSING_API void ProcessAudio(const float* InBuffer, const int32 InNumSamples, float* OutBuffer, const float* InKeyBuffer = nullptr);
 		SIGNALPROCESSING_API void ProcessAudio(const float* InBuffer, const int32 InNumSamples, float* OutBuffer, const float* InKeyBuffer, float* OutEnvelope);
 		 
+		// For single channels of audio OR non-interleaved blocks of multichannel audio.
+		SIGNALPROCESSING_API void ProcessAudio(const float* const* const InBuffers, const int32 InNumSamples, float* const* OutBuffers, const float* const* const InKeyBuffers, float* const* OutEnvelopes);
 
 	protected:
-
 		SIGNALPROCESSING_API float ComputeGain(const float InEnvFollowerDb);
+		SIGNALPROCESSING_API void ComputeGains(float* InEnvFollowerDbOutGain, const int32 InNumSamples);
 
 		// Process key frame, returning true if should continue processing
 		// (Returns false in audition mode and writes straight to output).
@@ -92,6 +94,8 @@ namespace Audio
 
 		EDynamicsProcessingMode::Type ProcessingMode;
 
+		float SlopeFactor;
+
 		// Peak mode of envelope followers
 		EPeakMode::Type EnvelopeFollowerPeakMode;
 
@@ -101,9 +105,6 @@ namespace Audio
 		// Envelope followers
 		TArray<FInlineEnvelopeFollower> EnvFollower;
 
-		// Points in the knee used for lagrangian interpolation
-		TArray<FVector2D> KneePoints;
-
 		// Channel values of cached detector sample
 		TArray<float> DetectorOuts;
 
@@ -158,5 +159,26 @@ namespace Audio
 		static constexpr float UpwardsCompressionMaxGain = 36.0f;
 
 		static constexpr float MaxLookaheadMsec = 100.0f;
+
+	private:
+		void CalculateSlope();
+
+		void CalculateKnee();
+
+		// Points in the knee used for lagrangian interpolation
+		struct FKneePoint
+		{
+			float X{ 0.0f };
+			float Y{ 0.0f };
+		};
+		TArray<FKneePoint> KneePoints;
+
+		// For optimized LagrangianInterpolation on blocks...
+		float Denominator0Minus1;
+
+		// For optimized LagrangianInterpolation on blocks...
+		float Denominator1Minus0;
+
 	};
+
 }
