Vectorized audio conversion from float to pcm16 and vica versa with an optimized version for NEON and added a ~10% faster spicialization of ArrayMixIn for NEON

diff --git a/Engine/Plugins/Experimental/MotoSynth/Source/MotoSynth/Private/MotoSynthSourceAsset.cpp b/Engine/Plugins/Experimental/MotoSynth/Source/MotoSynth/Private/MotoSynthSourceAsset.cpp
--- a/Engine/Plugins/Experimental/MotoSynth/Source/MotoSynth/Private/MotoSynthSourceAsset.cpp
+++ b/Engine/Plugins/Experimental/MotoSynth/Source/MotoSynth/Private/MotoSynthSourceAsset.cpp
@@ -339,10 +339,7 @@ void UMotoSynthSource::FilterSourceDataForAnalysis()
 	float* ScratchDataBufferPtr = ScratchBuffer.GetData();
 
 	// Convert the source data to floats
-	for (int32 FrameIndex = 0; FrameIndex < SourceDataPCM.Num(); ++FrameIndex)
-	{
-		ScratchDataBufferPtr[FrameIndex] = (float)SourceDataPCM[FrameIndex] / 32767.0f;
-	}
+	Audio::ArrayPcm16ToFloat(MakeArrayView(SourceDataPCM), MakeArrayView(ScratchDataBufferPtr, SourceDataPCM.Num()));
 
 	// Filter the audio source and write the output to the analysis buffer. Do not modify the source audio data.
 	if (bEnableFilteringForAnalysis)
@@ -550,10 +547,7 @@ void UMotoSynthSource::WriteAnalysisBufferToWaveFile()
 		TArray<int16> AnalysisBufferInt16;
 		AnalysisBufferInt16.AddUninitialized(AnalysisBuffer.Num());
 
-		for (int32 i = 0; i < AnalysisBufferInt16.Num(); ++i)
-		{
-			AnalysisBufferInt16[i] = AnalysisBuffer[i] * 32767.0f;
-		}
+		Audio::ArrayFloatToPcm16(MakeArrayView(AnalysisBuffer), MakeArrayView(AnalysisBufferInt16));
 
 		Audio::TSampleBuffer<> BufferToWrite(AnalysisBufferInt16.GetData(), AnalysisBufferInt16.Num(), 1, SourceSampleRate);
 
diff --git a/Engine/Plugins/Media/ElectraCodecs/Source/ElectraDecoders/ElectraDecoders.Build.cs b/Engine/Plugins/Media/ElectraCodecs/Source/ElectraDecoders/ElectraDecoders.Build.cs
--- a/Engine/Plugins/Media/ElectraCodecs/Source/ElectraDecoders/ElectraDecoders.Build.cs
+++ b/Engine/Plugins/Media/ElectraCodecs/Source/ElectraDecoders/ElectraDecoders.Build.cs
@@ -27,7 +27,8 @@ namespace UnrealBuildTool.Rules
 
 			PrivateDependencyModuleNames.AddRange(
 				new string[] {
-					"Core"
+					"Core",
+					"SignalProcessing"
 				});
 
 			PublicIncludePathModuleNames.AddRange(
diff --git a/Engine/Plugins/Media/ElectraCodecs/Source/ElectraDecoders/Private/DX/AudioDecoderAAC_DX.cpp b/Engine/Plugins/Media/ElectraCodecs/Source/ElectraDecoders/Private/DX/AudioDecoderAAC_DX.cpp
--- a/Engine/Plugins/Media/ElectraCodecs/Source/ElectraDecoders/Private/DX/AudioDecoderAAC_DX.cpp
+++ b/Engine/Plugins/Media/ElectraCodecs/Source/ElectraDecoders/Private/DX/AudioDecoderAAC_DX.cpp
@@ -6,7 +6,7 @@
 #include "DX/DecoderErrors_DX.h"
 #include "ElectraDecodersUtils.h"
 #include "Utils/MPEG/ElectraUtilsMPEGAudio.h"
-
+#include "DSP/FloatArrayMath.h"
 #include "IElectraDecoderFeaturesAndOptions.h"
 #include "IElectraDecoderOutputAudio.h"
 
@@ -939,10 +939,8 @@ bool FElectraAudioDecoderAAC_DX::ConvertDecoderOutput()
 	else
 	{
 		const int16* Source = reinterpret_cast<const int16*>(pDecompressedData);
-		for(int32 i=0, iMax=NewOutput->NumChannels * NewOutput->NumFrames; i<iMax; ++i)
-		{
-			*Dest++ = *Source++ / 32768.0f;
-		}
+		Audio::ArrayPcm16ToFloat(MakeArrayView(Source, NewOutput->NumChannels* NewOutput->NumFrames)
+			, MakeArrayView(Dest, NewOutput->NumChannels* NewOutput->NumFrames));
 	}
 	DecodedLinearOutputBuffer->Unlock();
 	CurrentOutput = MoveTemp(NewOutput);
diff --git a/Engine/Plugins/Media/PixelStreaming/Source/PixelStreaming/Private/PixelStreamingAudioComponent.cpp b/Engine/Plugins/Media/PixelStreaming/Source/PixelStreaming/Private/PixelStreamingAudioComponent.cpp
--- a/Engine/Plugins/Media/PixelStreaming/Source/PixelStreaming/Private/PixelStreamingAudioComponent.cpp
+++ b/Engine/Plugins/Media/PixelStreaming/Source/PixelStreaming/Private/PixelStreamingAudioComponent.cpp
@@ -250,12 +250,7 @@ int32 FWebRTCSoundGenerator::OnGenerateAudio(float* OutAudio, int32 NumSamples)
 		int32 NumSamplesToCopy = FGenericPlatformMath::Min(NumSamples, Buffer.Num());
 
 		// Copy from local buffer into OutAudio if we have enough samples
-		for (int i = 0; i < NumSamplesToCopy; i++)
-		{
-			//TODO: vectorize!
-			*OutAudio = Buffer[i] / 32767.0f;
-			OutAudio++;
-		}
+		Audio::ArrayPcm16ToFloat(TArrayView(Buffer, NumSamplesToCopy), TArrayView(OutAudio, NumSamplesToCopy));
 
 		// Remove front NumSamples from the local buffer
 		Buffer.RemoveAt(0, NumSamplesToCopy, false);
diff --git a/Engine/Plugins/Online/OnlineSubsystemUtils/Source/OnlineSubsystemUtils/Public/VoicePacketBuffer.h b/Engine/Plugins/Online/OnlineSubsystemUtils/Source/OnlineSubsystemUtils/Public/VoicePacketBuffer.h
--- a/Engine/Plugins/Online/OnlineSubsystemUtils/Source/OnlineSubsystemUtils/Public/VoicePacketBuffer.h
+++ b/Engine/Plugins/Online/OnlineSubsystemUtils/Source/OnlineSubsystemUtils/Public/VoicePacketBuffer.h
@@ -5,6 +5,7 @@
 #include "Containers/Queue.h"
 #include "HAL/ThreadSafeBool.h"
 #include "Net/VoiceConfig.h"
+#include "DSP/FloatArrayMath.h"
 
 // Set this to 1 to enforce a critical section between FVoicePacketBuffer::PopAudio and FVoicePacketBuffer::SubmitPacket.
 #define SCOPELOCK_VOICE_PACKET_BUFFER 0
@@ -92,10 +93,7 @@ struct FSortedVoicePacketNode
 				SamplesLeft = BufferNumSamples;
 				//Convert to float.
 				int16* BufferPtr = (int16*)InBuffer;
-				for (int32 Index = 0; Index < BufferNumSamples; Index++)
-				{
-					AudioBuffer[Index] = ((float)BufferPtr[Index]) / 32767.0f;
-				}
+				Audio::ArrayPcm16ToFloat(MakeArrayView(BufferPtr, BufferNumSamples), MakeArrayView(AudioBuffer, BufferNumSamples));
 				break;
 			}
 
diff --git a/Engine/Plugins/Runtime/Metasound/Source/MetasoundEngine/Private/MetasoundWaveProxyReader.cpp b/Engine/Plugins/Runtime/Metasound/Source/MetasoundEngine/Private/MetasoundWaveProxyReader.cpp
--- a/Engine/Plugins/Runtime/Metasound/Source/MetasoundEngine/Private/MetasoundWaveProxyReader.cpp
+++ b/Engine/Plugins/Runtime/Metasound/Source/MetasoundEngine/Private/MetasoundWaveProxyReader.cpp
@@ -140,17 +140,10 @@ namespace Metasound
 		{
 			SampleConversionBuffer.SetNumUninitialized(InBuffer.Num(), false /* bAllowShrinking */);
 
-			// Convert 16 bit pcm to 32 bit float.
-			constexpr float Scalar = 1.f / 32768.f;
 			const int16* Src = InBuffer.GetData();
 			float* Dst = SampleConversionBuffer.GetData();
-			int32 Num = InBuffer.Num();
-
-			// Convert 1 sample at a time, slow.
-			for (int32 i = 0; i < Num; ++i)
-			{
-				*Dst++ = static_cast<float>(*Src++) * Scalar;
-			}
+			const int32 Num = InBuffer.Num();
+			Audio::ArrayPcm16ToFloat(MakeArrayView(Src, Num), MakeArrayView(Dst, Num));
 
 			return PushAudioInternal(InDetails, SampleConversionBuffer);
 		}
diff --git a/Engine/Plugins/Runtime/Synthesis/Source/SynthesisEditor/Private/AudioImpulseResponseAsset.cpp b/Engine/Plugins/Runtime/Synthesis/Source/SynthesisEditor/Private/AudioImpulseResponseAsset.cpp
--- a/Engine/Plugins/Runtime/Synthesis/Source/SynthesisEditor/Private/AudioImpulseResponseAsset.cpp
+++ b/Engine/Plugins/Runtime/Synthesis/Source/SynthesisEditor/Private/AudioImpulseResponseAsset.cpp
@@ -9,6 +9,7 @@
 #include "Sound/SoundWaveProcedural.h"
 #include "ToolMenu.h"
 #include "ToolMenuSection.h"
+#include "DSP/FloatArrayMath.h"
 
 #include UE_INLINE_GENERATED_CPP_BY_NAME(AudioImpulseResponseAsset)
 
@@ -104,10 +105,7 @@ UObject* UAudioImpulseResponseFactory::FactoryCreateNew(UClass* Class, UObject*
 			const int16* InputBuffer = (int16*)ImportedSoundWaveData.GetData();
 			float* OutputBuffer = NewAsset->ImpulseResponse.GetData();
 
-			for (int32 i = 0; i < NumSamples; ++i)
-			{
-				OutputBuffer[i] = static_cast<float>(InputBuffer[i]) / 32768.0f;
-			}
+			Audio::ArrayPcm16ToFloat(MakeArrayView(InputBuffer, NumSamples), MakeArrayView(OutputBuffer, NumSamples));
 		}
 
 #if WITH_EDITORONLY_DATA
diff --git a/Engine/Plugins/Runtime/Synthesis/Source/SynthesisEditor/SynthesisEditor.Build.cs b/Engine/Plugins/Runtime/Synthesis/Source/SynthesisEditor/SynthesisEditor.Build.cs
--- a/Engine/Plugins/Runtime/Synthesis/Source/SynthesisEditor/SynthesisEditor.Build.cs
+++ b/Engine/Plugins/Runtime/Synthesis/Source/SynthesisEditor/SynthesisEditor.Build.cs
@@ -34,6 +34,7 @@ namespace UnrealBuildTool.Rules
 					"DeveloperSettings",
 					"UMG",
 					"UMGEditor",
+					"SignalProcessing"
 				}
 			);
 
diff --git a/Engine/Source/Runtime/AVEncoder/AVEncoder.Build.cs b/Engine/Source/Runtime/AVEncoder/AVEncoder.Build.cs
--- a/Engine/Source/Runtime/AVEncoder/AVEncoder.Build.cs
+++ b/Engine/Source/Runtime/AVEncoder/AVEncoder.Build.cs
@@ -19,7 +19,8 @@ public class AVEncoder : ModuleRules
 		// PrecompileForTargets = PrecompileTargetsType.None;
 
 		PrivateDependencyModuleNames.AddRange(new string[] {
-			"Engine"
+			"Engine",
+			"SignalProcessing"
 		});
 
 		PublicDependencyModuleNames.AddRange(new string[] {
diff --git a/Engine/Source/Runtime/AudioCodecEngine/AudioCodecEngine.Build.cs b/Engine/Source/Runtime/AudioCodecEngine/AudioCodecEngine.Build.cs
--- a/Engine/Source/Runtime/AudioCodecEngine/AudioCodecEngine.Build.cs
+++ b/Engine/Source/Runtime/AudioCodecEngine/AudioCodecEngine.Build.cs
@@ -13,6 +13,7 @@ namespace UnrealBuildTool.Rules
 					"CoreUObject",
                     "Engine",
 					"AudioExtensions",
+					"SignalProcessing",
 				//	"TargetPlatform"
 				}
 			);
diff --git a/Engine/Source/Runtime/AudioCodecEngine/Private/SoundWaveProxyReader.cpp b/Engine/Source/Runtime/AudioCodecEngine/Private/SoundWaveProxyReader.cpp
--- a/Engine/Source/Runtime/AudioCodecEngine/Private/SoundWaveProxyReader.cpp
+++ b/Engine/Source/Runtime/AudioCodecEngine/Private/SoundWaveProxyReader.cpp
@@ -534,14 +534,10 @@ FSoundWaveProxyReader::EDecodeResult FSoundWaveProxyReader::Decode()
 			break;
 		}
 
-		int32 NumSamplesStreamed = NumBytesStreamed / sizeof(int16);
-		int32 NumFramesStreamed = NumSamplesStreamed / NumChannels;
+		const int32 NumSamplesStreamed = NumBytesStreamed / sizeof(int16);
+		const int32 NumFramesStreamed = NumSamplesStreamed / NumChannels;
 
-		constexpr float Scalar = 1.f / 32768.f;
-		for (int32 SampleIdx = 0; SampleIdx < NumSamplesStreamed; ++SampleIdx)
-		{
-			SampleConversionBuffer[SampleIdx] = static_cast<float>(ResidualBuffer[SampleIdx]) * Scalar;
-		}
+		Audio::ArrayPcm16ToFloat(MakeArrayView(ResidualBuffer, NumSamplesStreamed), MakeArrayView(SampleConversionBuffer, NumSamplesStreamed));
 
 		const float* SampleData = SampleConversionBuffer.GetData();
 		DecoderOutput.Push(SampleData, NumSamplesStreamed);
diff --git a/Engine/Source/Runtime/AudioMixer/Private/AudioMixerSourceBuffer.cpp b/Engine/Source/Runtime/AudioMixer/Private/AudioMixerSourceBuffer.cpp
--- a/Engine/Source/Runtime/AudioMixer/Private/AudioMixerSourceBuffer.cpp
+++ b/Engine/Source/Runtime/AudioMixer/Private/AudioMixerSourceBuffer.cpp
@@ -5,6 +5,7 @@
 #include "ContentStreaming.h"
 #include "AudioDecompress.h"
 #include "Misc/ScopeTryLock.h"
+#include "DSP/FloatArrayMath.h"
 
 namespace Audio
 {
@@ -288,15 +289,9 @@ namespace Audio
 				int16* CachedBufferPtr1 = (int16*)(CachedRealtimeFirstBuffer.GetData() + BufferSize);
 				float* AudioData0 = SourceVoiceBuffers[0]->AudioData.GetData();
 				float* AudioData1 = SourceVoiceBuffers[1]->AudioData.GetData();
-				for (uint32 Sample = 0; Sample < NumSamples; ++Sample)
-				{
-					AudioData0[Sample] = CachedBufferPtr0[Sample] / 32768.0f;
-				}
 
-				for (uint32 Sample = 0; Sample < NumSamples; ++Sample)
-				{
-					AudioData1[Sample] = CachedBufferPtr1[Sample] / 32768.0f;
-				}
+				Audio::ArrayPcm16ToFloat(MakeArrayView(CachedBufferPtr0, NumSamples), MakeArrayView(AudioData0, NumSamples));
+				Audio::ArrayPcm16ToFloat(MakeArrayView(CachedBufferPtr1, NumSamples), MakeArrayView(AudioData1, NumSamples));
 
 				// Submit the already decoded and cached audio buffers
 				SubmitBuffer(SourceVoiceBuffers[0]);
@@ -310,12 +305,8 @@ namespace Audio
 				SourceVoiceBuffers[0]->AudioData.AddZeroed(NumSamples);
 
 				int16* CachedBufferPtr0 = (int16*)CachedRealtimeFirstBuffer.GetData();
-
 				float* AudioData0 = SourceVoiceBuffers[0]->AudioData.GetData();
-				for (uint32 Sample = 0; Sample < NumSamples; ++Sample)
-				{
-					AudioData0[Sample] = CachedBufferPtr0[Sample] / 32768.0f;
-				}
+				Audio::ArrayPcm16ToFloat(MakeArrayView(CachedBufferPtr0, NumSamples), MakeArrayView(AudioData0, NumSamples));
 
 				// Submit the already decoded and cached audio buffers
 				SubmitBuffer(SourceVoiceBuffers[0]);
diff --git a/Engine/Source/Runtime/AudioMixer/Private/AudioMixerSourceDecode.cpp b/Engine/Source/Runtime/AudioMixer/Private/AudioMixerSourceDecode.cpp
--- a/Engine/Source/Runtime/AudioMixer/Private/AudioMixerSourceDecode.cpp
+++ b/Engine/Source/Runtime/AudioMixer/Private/AudioMixerSourceDecode.cpp
@@ -10,6 +10,7 @@
 #include "Async/Async.h"
 #include "AudioDecompress.h"
 #include "Sound/SoundGenerator.h"
+#include "DSP/FloatArrayMath.h"
 
 static int32 ForceSyncAudioDecodesCvar = 0;
 FAutoConsoleVariableRef CVarForceSyncAudioDecodes(
@@ -125,13 +126,8 @@ public:
 						check(NumBytesWritten <= ByteSize);
 
 						ProceduralResult.NumSamplesWritten = NumBytesWritten / sizeof(int16);
-
-						// Convert the buffer to float
-						int16* DecodedBufferPtr = (int16*)DecodeBuffer.GetData();
-						for (int32 SampleIndex = 0; SampleIndex < ProceduralResult.NumSamplesWritten; ++SampleIndex)
-						{
-							ProceduralTaskData.AudioData[SampleIndex] = (float)(DecodedBufferPtr[SampleIndex]) / 32768.0f;
-						}
+						Audio::ArrayPcm16ToFloat(MakeArrayView((int16*)DecodeBuffer.GetData(), ProceduralResult.NumSamplesWritten)
+							, MakeArrayView(ProceduralTaskData.AudioData, ProceduralResult.NumSamplesWritten));
 					}
 					else
 					{
@@ -195,15 +191,8 @@ public:
 				}
 
 				// Convert the decoded PCM data into a float buffer while still in the async task
-				int32 SampleIndex = 0;
-				int16* DecodedBufferPtr = (int16*)DecodeBuffer.GetData();
-				for (int32 Frame = 0; Frame < DecodeTaskData.NumFramesToDecode; ++Frame)
-				{
-					for (int32 Channel = 0; Channel < NumChannels; ++Channel, ++SampleIndex)
-					{
-						DecodeTaskData.AudioData[SampleIndex] = (float)(DecodedBufferPtr[SampleIndex]) / 32768.0f;
-					}
-				}
+				Audio::ArrayPcm16ToFloat(MakeArrayView((int16*)DecodeBuffer.GetData(), DecodeTaskData.NumFramesToDecode * NumChannels)
+					, MakeArrayView(DecodeTaskData.AudioData, DecodeTaskData.NumFramesToDecode* NumChannels));
 			}
 			break;
 		}
diff --git a/Engine/Source/Runtime/AudioMixer/Private/FileDecoder.cpp b/Engine/Source/Runtime/AudioMixer/Private/FileDecoder.cpp
--- a/Engine/Source/Runtime/AudioMixer/Private/FileDecoder.cpp
+++ b/Engine/Source/Runtime/AudioMixer/Private/FileDecoder.cpp
@@ -4,6 +4,7 @@
 #include "OpusAudioInfo.h"
 #include "VorbisAudioInfo.h"
 #include "HAL/PlatformFileManager.h"
+#include "DSP/FloatArrayMath.h"
 
 FAudioFileReader::FAudioFileReader(const FString& InPath)
 {
@@ -53,10 +54,7 @@ bool FAudioFileReader::PopAudio(float* OutAudio, int32 NumSamples)
 	bool bIsFinished = Decompressor->ReadCompressedData((uint8*) DecompressionBuffer.GetData(), false, NumSamples * sizeof(Audio::DefaultUSoundWaveSampleType));
 
 	// Convert to float:
-	for (int32 Index = 0; Index < NumSamples; Index++)
-	{
-		OutAudio[Index] = ((float)DecompressionBuffer[Index]) / 32768.0f;
-	}
+	Audio::ArrayPcm16ToFloat(MakeArrayView((int16*)DecompressionBuffer.GetData(), NumSamples), MakeArrayView(OutAudio, NumSamples));
 
 	return bIsFinished;
 }
diff --git a/Engine/Source/Runtime/MediaUtils/MediaUtils.Build.cs b/Engine/Source/Runtime/MediaUtils/MediaUtils.Build.cs
--- a/Engine/Source/Runtime/MediaUtils/MediaUtils.Build.cs
+++ b/Engine/Source/Runtime/MediaUtils/MediaUtils.Build.cs
@@ -13,6 +13,7 @@ namespace UnrealBuildTool.Rules
 					"CoreUObject",
 					"ImageWrapper",
 					"Media",
+					"SignalProcessing"
 				});
 
 			PrivateDependencyModuleNames.AddRange(
diff --git a/Engine/Source/Runtime/MediaUtils/Private/MediaAudioResampler.cpp b/Engine/Source/Runtime/MediaUtils/Private/MediaAudioResampler.cpp
--- a/Engine/Source/Runtime/MediaUtils/Private/MediaAudioResampler.cpp
+++ b/Engine/Source/Runtime/MediaUtils/Private/MediaAudioResampler.cpp
@@ -2,7 +2,7 @@
 
 #include "MediaAudioResampler.h"
 #include "MediaUtilsPrivate.h"
-
+#include "DSP/FloatArrayMath.h"
 #include "Algo/Reverse.h"
 #include "IMediaAudioSample.h"
 
@@ -547,12 +547,10 @@ bool FMediaAudioResampler::SetInput(const TSharedPtr<IMediaAudioSample, ESPMode:
 			break;
 
 		case EMediaAudioSampleFormat::Int16:
-			while (FloatData < FloatDataEnd)
-			{
-				*FloatData++ = (float)(*(const int16*)Buffer / 32768.0f);
-				Buffer += sizeof(int16);
-			}
+		{
+			Audio::ArrayPcm16ToFloat(MakeArrayView((const int16*)Buffer, NumSamples), MakeArrayView(FloatData, NumSamples));
 			break;
+		}
 
 		case EMediaAudioSampleFormat::Int32:
 			while (FloatData < FloatDataEnd)
diff --git a/Engine/Source/Runtime/SignalProcessing/Private/FloatArrayMath.cpp b/Engine/Source/Runtime/SignalProcessing/Private/FloatArrayMath.cpp
--- a/Engine/Source/Runtime/SignalProcessing/Private/FloatArrayMath.cpp
+++ b/Engine/Source/Runtime/SignalProcessing/Private/FloatArrayMath.cpp
@@ -2030,28 +2030,33 @@ namespace Audio
 		}
 		else
 		{
-			const int32 NumToSimd = Num & MathIntrinsics::SimdMask;
-			const int32 NumNotToSimd = Num & MathIntrinsics::NotSimdMask;
-
-			if (NumToSimd)
+			VectorRegister4Float GainVector = VectorLoadFloat1(&Gain);
+			int32 i = 0;
+#if PLATFORM_ENABLE_VECTORINTRINSICS_NEON
+			// this approach is ~10% faster than the other one on NEON
+			for (; i < Num; i += 16)
+			{
+				float32x4x4_t Input = vld1q_f32_x4(&InData[i]);
+				float32x4x4_t Output = vld1q_f32_x4(&InOutData[i]);
+				Output.val[0] = VectorMultiplyAdd(Input.val[0], GainVector, Output.val[0]);
+				Output.val[1] = VectorMultiplyAdd(Input.val[1], GainVector, Output.val[1]);
+				Output.val[2] = VectorMultiplyAdd(Input.val[2], GainVector, Output.val[2]);
+				Output.val[3] = VectorMultiplyAdd(Input.val[3], GainVector, Output.val[3]);
+				vst1q_f32_x4(&InOutData[i], Output);
+			}
+#else
+			for (; i < Num; i += AUDIO_NUM_FLOATS_PER_VECTOR_REGISTER)
 			{
-				VectorRegister4Float GainVector = VectorLoadFloat1(&Gain);
-
-				for (int32 i = 0; i < NumToSimd; i += AUDIO_NUM_FLOATS_PER_VECTOR_REGISTER)
-				{
-					VectorRegister4Float Output = VectorLoad(&InOutData[i]);
-					VectorRegister4Float Input = VectorLoad(&InData[i]);
-					Output = VectorMultiplyAdd(Input, GainVector, Output);
-					VectorStore(Output, &InOutData[i]);
-				}
+				VectorRegister4Float Input = VectorLoad(&InData[i]);
+				VectorRegister4Float Output = VectorLoad(&InOutData[i]);
+				Output = VectorMultiplyAdd(Input, GainVector, Output);
+				VectorStore(Output, &InOutData[i]);
 			}
+#endif //~PLATFORM_ENABLE_VECTORINTRINSICS_NEON
 
-			if (NumNotToSimd)
+			for (; i < Num; ++i)
 			{
-				for (int32 i = NumToSimd; i < Num; ++i)
-				{
-					InOutData[i] += InData[i] * Gain;
-				}
+				InOutData[i] += InData[i] * Gain;
 			}
 		}
 	}
@@ -2214,43 +2219,47 @@ namespace Audio
 	{
 		CSV_SCOPED_TIMING_STAT(Audio_Dsp, ArrayFloatToPcm16);
 
-		check(InView.Num() == OutView.Num());
+		check(OutView.Num() >= InView.Num());
 
 		const int32 Num = InView.Num();
-		
-		const int32 NumToSimd = Num & MathIntrinsics::SimdMask;
-		const int32 NumNotToSimd = Num & MathIntrinsics::NotSimdMask;
 
 		const float* InputPtr = InView.GetData();
 		int16* OutPtr = OutView.GetData();
 
 		constexpr float ConversionValue = static_cast<float>(TNumericLimits<int16>::Max());
-
-		if(NumToSimd)
+		const VectorRegister4Float Multiplier = VectorSetFloat1(ConversionValue);
+		int32 i = 0;
+#if PLATFORM_ENABLE_VECTORINTRINSICS_NEON
+		for (; i < Num; i += 8)
+		{
+			const VectorRegister4Float InVector1 = VectorLoad(&InputPtr[i]);
+			const VectorRegister4Float InVector2 = VectorLoad(&InputPtr[i + 4]);
+			const VectorRegister4Float ScaledVector1 = VectorMultiply(InVector1, Multiplier);
+			const VectorRegister4Float ScaledVector2 = VectorMultiply(InVector2, Multiplier);
+			const VectorRegister4Int IntVector1 = VectorFloatToInt(ScaledVector1);
+			const VectorRegister4Int IntVector2 = VectorFloatToInt(ScaledVector2);
+			const int16x8_t Result = vmovn_high_s32(vmovn_u32(IntVector1), IntVector2);
+			vst1q_s16(&OutPtr[i], Result);
+		}
+#else
+		for (; i < Num; i += AUDIO_NUM_FLOATS_PER_VECTOR_REGISTER)
 		{
-			const VectorRegister4Float ConversionVector = VectorSetFloat1(ConversionValue);
-
-			for (int32 i = 0; i < NumToSimd; i += AUDIO_NUM_FLOATS_PER_VECTOR_REGISTER)
-			{
-				const VectorRegister4Float InVector = VectorLoad(&InputPtr[i]);
-				const VectorRegister4Float ScaledVector = VectorMultiply(InVector, ConversionVector);
-				const VectorRegister4Int IntVector = VectorFloatToInt(ScaledVector);
+			const VectorRegister4Float InVector = VectorLoad(&InputPtr[i]);
+			const VectorRegister4Float ScaledVector = VectorMultiply(InVector, Multiplier);
+			const VectorRegister4Int IntVector = VectorFloatToInt(ScaledVector);
 
-				const AlignedFloat4 ScaledFloatArray(ScaledVector);
+			const AlignedFloat4 ScaledFloatArray(ScaledVector);
 
-				OutPtr[i] =		(int16)ScaledFloatArray[0];
-				OutPtr[i + 1] =	(int16)ScaledFloatArray[1];
-				OutPtr[i + 2] =	(int16)ScaledFloatArray[2];
-				OutPtr[i + 3] =	(int16)ScaledFloatArray[3];
-			}
+			OutPtr[i] =		(int16)ScaledFloatArray[0];
+			OutPtr[i + 1] =	(int16)ScaledFloatArray[1];
+			OutPtr[i + 2] =	(int16)ScaledFloatArray[2];
+			OutPtr[i + 3] =	(int16)ScaledFloatArray[3];
 		}
+#endif //~PLATFORM_ENABLE_VECTORINTRINSICS_NEON
 
-		if(NumNotToSimd)
+		for (; i < Num; i++)
 		{
-			for (int32 i = NumToSimd; i < Num; i++)
-			{
-				OutPtr[i] = (int16)(InputPtr[i] * ConversionValue);
-			}
+			OutPtr[i] = (int16)(InputPtr[i] * ConversionValue);
 		}
 	}
 	
@@ -2258,43 +2267,47 @@ namespace Audio
 	{
 		CSV_SCOPED_TIMING_STAT(Audio_Dsp, ArrayPcm16ToFloat);
 
-		check(InView.Num() == OutView.Num());
+		check(OutView.Num() >= InView.Num());
 
 		const int32 Num = InView.Num();
-		
-		const int32 NumToSimd = Num & MathIntrinsics::SimdMask;
-		const int32 NumNotToSimd = Num & MathIntrinsics::NotSimdMask;
 
 		const int16* InputPtr = InView.GetData();
 		float* OutPtr = OutView.GetData();
 
 		constexpr float ConversionValue = 1.f / static_cast<float>(TNumericLimits<int16>::Max());
+		const VectorRegister4Float Multiplier = VectorSetFloat1(ConversionValue);
 
-		if(NumToSimd)
+		int32 i = 0;
+#if PLATFORM_ENABLE_VECTORINTRINSICS_NEON
+		for (; i < Num; i += 8)
 		{
-			const VectorRegister4Float ConversionVector = VectorSetFloat1(ConversionValue);
-			AlignedFloat4 FloatArray(GlobalVectorConstants::FloatZero);
+			int16x8_t Data = vld1q_s16(&InputPtr[i]);
+			int32x4_t VecA = vmovl_s16(vget_low_s16(Data));
+			int32x4_t VecB = vmovl_high_s16(Data);
+			VectorRegister4Float FloatVecA = VectorMultiply(vcvtq_f32_s32(VecA), Multiplier);
+			VectorRegister4Float FloatVecB = VectorMultiply(vcvtq_f32_s32(VecB), Multiplier);
+			VectorStore(FloatVecA, &OutPtr[i]);
+			VectorStore(FloatVecB, &OutPtr[i + 4]);
+		}
+#else
+		AlignedFloat4 FloatArray(GlobalVectorConstants::FloatZero);
+		for (; i < Num; i += AUDIO_NUM_FLOATS_PER_VECTOR_REGISTER)
+		{
+			FloatArray[0] = (float)InputPtr[i];
+			FloatArray[1] = (float)InputPtr[i + 1];
+			FloatArray[2] = (float)InputPtr[i + 2];
+			FloatArray[3] = (float)InputPtr[i + 3];
 
-			for (int32 i = 0; i < NumToSimd; i += AUDIO_NUM_FLOATS_PER_VECTOR_REGISTER)
-			{
-				FloatArray[0] = (float)InputPtr[i];
-				FloatArray[1] = (float)InputPtr[i + 1];
-				FloatArray[2] = (float)InputPtr[i + 2];
-				FloatArray[3] = (float)InputPtr[i + 3];
+			const VectorRegister4Float InVector = FloatArray.ToVectorRegister();
+			const VectorRegister4Float ScaledVector = VectorMultiply(InVector, Multiplier);
 
-				const VectorRegister4Float InVector = FloatArray.ToVectorRegister();
-				const VectorRegister4Float ScaledVector = VectorMultiply(InVector, ConversionVector);
-			
-				VectorStore(ScaledVector, &OutPtr[i]);
-			}
+			VectorStore(ScaledVector, &OutPtr[i]);
 		}
+#endif //~PLATFORM_ENABLE_VECTORINTRINSICS_NEON
 
-		if(NumNotToSimd)
+		for (; i < Num; i++)
 		{
-			for (int32 i = NumToSimd; i < Num; i++)
-			{
-				OutPtr[i] = (float)InputPtr[i] * ConversionValue;
-			}
+			OutPtr[i] = (float)InputPtr[i] * ConversionValue;
 		}
 	}
 	
diff --git a/Engine/Source/Runtime/SignalProcessing/Public/SampleBuffer.h b/Engine/Source/Runtime/SignalProcessing/Public/SampleBuffer.h
--- a/Engine/Source/Runtime/SignalProcessing/Public/SampleBuffer.h
+++ b/Engine/Source/Runtime/SignalProcessing/Public/SampleBuffer.h
@@ -4,6 +4,7 @@
 
 #include "CoreMinimal.h"
 #include "DSP/BufferVectorOperations.h"
+#include "DSP/FloatArrayMath.h"
 
 namespace Audio
 {
@@ -116,10 +117,7 @@ namespace Audio
 			else if constexpr(std::is_same_v<SampleType, float>)
 			{
 				// Convert from int to float:
-				for (int32 SampleIndex = 0; SampleIndex < NumSamples; SampleIndex++)
-				{
-					RawPCMData[SampleIndex] = (InBufferPtr[SampleIndex] / 32767.0f);
-				}
+				Audio::ArrayPcm16ToFloat(MakeArrayView(InBufferPtr, NumSamples), RawPCMData);
 			}
 			else
 			{
@@ -169,18 +167,12 @@ namespace Audio
 			else if constexpr(std::is_same_v<SampleType, int16> && std::is_same_v<OtherSampleType, float>)
 			{
 				// Convert from float to int:
-				for (int32 SampleIndex = 0; SampleIndex < NumSamples; SampleIndex++)
-				{
-					RawPCMData[SampleIndex] = (int16)(Other.RawPCMData[SampleIndex] * 32767.0f);
-				}
+				Audio::ArrayFloatToPcm16(MakeArrayView(Other.RawPCMData), MakeArrayView(RawPCMData));
 			}
 			else if constexpr(std::is_same_v<SampleType, float> && std::is_same_v<OtherSampleType, int16>)
 			{
 				// Convert from int to float:
-				for (int32 SampleIndex = 0; SampleIndex < NumSamples; SampleIndex++)
-				{
-					RawPCMData[SampleIndex] = (float)((float)Other.RawPCMData[SampleIndex]) / 32767.0f;
-				}
+				Audio::ArrayPcm16ToFloat(MakeArrayView(Other.RawPCMData), MakeArrayView(RawPCMData));
 			}
 			else
 			{
@@ -224,18 +216,12 @@ namespace Audio
 				if constexpr(std::is_same_v<SampleType, int16> && std::is_same_v<OtherSampleType, float>)
 				{
 					// Convert from float to int:
-					for (int32 SampleIndex = 0; SampleIndex < InNumSamples; SampleIndex++)
-					{
-						RawPCMData[StartIndex + SampleIndex] = (int16)(InputBuffer[SampleIndex] * 32767.0f);
-					}
+					Audio::ArrayFloatToPcm16(MakeArrayView(InputBuffer, InNumSamples), MakeArrayView(&RawPCMData[StartIndex], InNumSamples));
 				}
 				else if constexpr(std::is_same_v<SampleType, float> && std::is_same_v<OtherSampleType, int16>)
 				{
 					// Convert from int to float:
-					for (int32 SampleIndex = 0; SampleIndex < InNumSamples; SampleIndex++)
-					{
-						RawPCMData[StartIndex + SampleIndex] = (float)InputBuffer[SampleIndex] / 32767.0f;
-					}
+					Audio::ArrayPcm16ToFloat(MakeArrayView(InputBuffer, InNumSamples), MakeArrayView(&RawPCMData[StartIndex], NumSamples));
 				}
 				else
 				{
