strengthen buffer overflow defense in interpolated multi tap delay line to prevent rounding errors on buffer indices

#rnx

diff --git a/Engine/Source/Runtime/SignalProcessing/Private/InterpolatedMultiTapDelay.cpp b/Engine/Source/Runtime/SignalProcessing/Private/InterpolatedMultiTapDelay.cpp
--- a/Engine/Source/Runtime/SignalProcessing/Private/InterpolatedMultiTapDelay.cpp
+++ b/Engine/Source/Runtime/SignalProcessing/Private/InterpolatedMultiTapDelay.cpp
@@ -8,81 +8,98 @@
 
 namespace Audio
 {
-	void FInterpolatedMultiTapDelay::Init(const int32 InBufferSizeSamples, const float InSampleRate)
+	void FInterpolatedMultiTapDelay::Init(const int32 InBufferSizeSamples)
 	{
 		WriteIndex = 0;
-		MsToSamples = InSampleRate / 1000.f;
 		DelayLine.Reset();
-		DelayLine.AddZeroed(AlignIndex(InBufferSizeSamples));
+		DelayLine.AddZeroed(InBufferSizeSamples);
+		WrapBuffer.Reset();
+		WrapBuffer.AddZeroed(InBufferSizeSamples);
 	}
 
-	void FInterpolatedMultiTapDelay::Advance(const FAlignedFloatBuffer& InSamples)
+	void FInterpolatedMultiTapDelay::Advance(const FAlignedFloatBuffer& InBuffer)
 	{
-		check (InSamples.Num() % AUDIO_NUM_FLOATS_PER_VECTOR_REGISTER == 0);
-		const int32 InNumSamples = InSamples.Num();
+		check (InBuffer.Num() % AUDIO_NUM_FLOATS_PER_VECTOR_REGISTER == 0);
+		const uint32 InNumSamples = InBuffer.Num();
+		const uint32 DelayBufferNumSamples = DelayLine.Num();
 
 		if (InNumSamples <= 0)
 		{
 			return;
 		}
 
-		if (InNumSamples + WriteIndex > DelayLine.Num())
+		if (InNumSamples + WriteIndex > DelayBufferNumSamples)
 		{
-			const int32 FirstHalfSize = DelayLine.Num() - WriteIndex;
+			const int32 FirstHalfSize = DelayBufferNumSamples - WriteIndex;
 			const int32 SecondHalfSize = InNumSamples - FirstHalfSize;
 
 			if (FirstHalfSize > 0)
 			{
-				FMemory::Memcpy(&DelayLine[WriteIndex], InSamples.GetData(), FirstHalfSize * sizeof(float));
+				FMemory::Memcpy(&DelayLine[WriteIndex], InBuffer.GetData(), FirstHalfSize * sizeof(float));
 			}
 			if (SecondHalfSize > 0)
 			{
-				FMemory::Memcpy(DelayLine.GetData(), &InSamples[FirstHalfSize], SecondHalfSize * sizeof(float));
+				FMemory::Memcpy(DelayLine.GetData(), &InBuffer[FirstHalfSize], SecondHalfSize * sizeof(float));
 			}
 			WriteIndex = SecondHalfSize;
 		}
 		else
 		{
-			FMemory::Memcpy(&DelayLine[WriteIndex], InSamples.GetData(), InNumSamples * sizeof(float));
+			FMemory::Memcpy(&DelayLine[WriteIndex], InBuffer.GetData(), InNumSamples * sizeof(float));
 			WriteIndex += InNumSamples;
-			WriteIndex = FMath::Wrap(WriteIndex, 0, DelayLine.Num());
+			WriteIndex = FMath::Wrap<uint32>(WriteIndex, 0, DelayLine.Num());
 		}
 	}
 
-	void FInterpolatedMultiTapDelay::Read(const float StartDelayMSec, const float EndDelayMSec, FAlignedFloatBuffer& OutSamples)
+	uint32 FInterpolatedMultiTapDelay::Read(const uint32 StartNumDelaySamples, const uint32 StartSampleFraction, const uint32 EndNumDelaySamples, FAlignedFloatBuffer& OutBuffer)
 	{
-		const int32 OutputSamples = OutSamples.Num();
-		check (OutputSamples % AUDIO_NUM_FLOATS_PER_VECTOR_REGISTER == 0);
+		const int32 OutputNumSamples = OutBuffer.Num();
+		const int32 DelayBufferNumSamples = DelayLine.Num();
 
 		// likely to only run on the first frame, if not configured with enough memory
-		if (OutputSamples > DelayLine.Num())
+		if (OutputNumSamples > DelayBufferNumSamples)
 		{
-			DelayLine.SetNumZeroed(OutputSamples);
+			DelayLine.SetNumZeroed(OutputNumSamples);
 		}
 
-		const int32 BufferSize = DelayLine.Num();
-		const int32 StartSample = AlignIndex(FMath::Wrap(WriteIndex - FMath::FloorToInt32(StartDelayMSec * MsToSamples), 0, BufferSize - 1));
-		const int32 EndSample = AlignIndex(FMath::Wrap(WriteIndex + OutputSamples - FMath::CeilToInt32(EndDelayMSec * MsToSamples), 0, BufferSize - 1));
-
-		float* OutputPtr = OutSamples.GetData();
+		if (OutputNumSamples > WrapBuffer.Num())
+		{
+			WrapBuffer.SetNumUninitialized(OutputNumSamples);
+		}
 		
-		if (StartSample >= EndSample)
+		if (DelayBufferNumSamples <= 0 || OutputNumSamples <= 0)
 		{
-			const int32 DelaySamples = EndSample - (StartSample - BufferSize);
+			return 0;
+		}
+
+		int32 StartSample = WriteIndex - StartNumDelaySamples;
+		int32 EndSample = WriteIndex + OutputNumSamples - EndNumDelaySamples;
 
-			// block wraps
-			const int32 FirstBlockDelaySamples = BufferSize - StartSample;
-			const int32 SecondBlockDelaySamples = DelaySamples - FirstBlockDelaySamples;
-			const int32 FirstBlockOutputSamples = AlignIndex(((float)FirstBlockDelaySamples / (float)DelaySamples) * OutputSamples);
-			const int32 SecondBlockOuptutSamples = OutputSamples - FirstBlockOutputSamples;
+		const float SampleStride = FMath::Clamp((float)(EndSample - StartSample) / (float)OutputNumSamples, 0.25f, 4.f);
+		const uint32 FixedSampleRate = (uint32)(SampleStride * 65536.f);
+
+		StartSample = FMath::Wrap(StartSample, 0, DelayBufferNumSamples - 1);
 		
-			ReadBlockInternal(StartSample, FirstBlockDelaySamples, FirstBlockOutputSamples, OutputPtr);
-			ReadBlockInternal(0, SecondBlockDelaySamples, SecondBlockOuptutSamples, OutputPtr + FirstBlockOutputSamples);
-		}
-		else
+		Resampler.CurrentFrameFraction = StartSampleFraction;
+		const int32 FramesNeeded = (int32)Resampler.SourceFramesNeeded(OutputNumSamples, FixedSampleRate);
+
+		float* SourceBuffer = &DelayLine[StartSample];
+		if (StartSample + FramesNeeded >= DelayBufferNumSamples)
 		{
-			ReadBlockInternal(StartSample, EndSample - StartSample, OutputSamples, OutputPtr);
+			constexpr int32 NumSafetyBufferSamples = 16;
+			
+			const int32 NumSamplesToEnd = DelayBufferNumSamples - StartSample;
+			const int32 SecondBufferNumSamples = FMath::Min((FramesNeeded - NumSamplesToEnd) + NumSafetyBufferSamples, DelayBufferNumSamples);
+			
+			FMemory::Memcpy(WrapBuffer.GetData(), &DelayLine[StartSample], NumSamplesToEnd * sizeof(float));
+			FMemory::Memcpy(&WrapBuffer[NumSamplesToEnd], DelayLine.GetData(), SecondBufferNumSamples * sizeof(float));
+
+			SourceBuffer = WrapBuffer.GetData();
 		}
+
+		Resampler.ResampleMono(OutputNumSamples, FixedSampleRate, SourceBuffer, OutBuffer.GetData());
+
+		return Resampler.CurrentFrameFraction;
 	}
 
 	void FInterpolatedMultiTapDelay::Reset()
@@ -94,55 +111,6 @@ namespace Audio
 
 	bool FInterpolatedMultiTapDelay::IsInitialized() const
 	{
-		return DelayLine.GetAllocatedSize() > 0 && MsToSamples > 1.f;
-	}
-
-	void FInterpolatedMultiTapDelay::ReadBlockInternal(const int32 StartSample, const int32 SamplesToRead, const int32 NumOutputSamples, float* OutSamples)
-	{
-		if (NumOutputSamples == 0 || SamplesToRead == 0)
-		{
-			return;
-		}
-		
-		check(StartSample + SamplesToRead <= DelayLine.Num())
-		check(NumOutputSamples % AUDIO_NUM_FLOATS_PER_VECTOR_REGISTER == 0)
-
-		const float SampleStride = (float)SamplesToRead / (float)NumOutputSamples;
-
-		// protect against overflows while wrapping without branching in the main loop
-		int32 MainLoopOutputSamples = NumOutputSamples;
-		int32 MainLoopReadSamples = SamplesToRead;
-		
-		if (StartSample + SamplesToRead == DelayLine.Num())
-		{
-			MainLoopOutputSamples -= AUDIO_NUM_FLOATS_PER_VECTOR_REGISTER;
-			MainLoopReadSamples -= FMath::RoundToInt32((float)AUDIO_NUM_FLOATS_PER_VECTOR_REGISTER * SampleStride);
-		}
-
-		ArrayInterpolate(&DelayLine[StartSample], OutSamples, MainLoopReadSamples, MainLoopOutputSamples);
-
-		float SampleIndex = MainLoopReadSamples;
-		
-		// wrap in the loop only if close to wrapping: protects against rounding up past last index
-		for (int32 OutputIndex = MainLoopOutputSamples; OutputIndex < NumOutputSamples; OutputIndex++)
-		{
-			const int32 LeftSample = FMath::FloorToInt32(SampleIndex);
-			int32 RightSample = FMath::CeilToInt32(SampleIndex);
-			
-			if (RightSample >= DelayLine.Num())
-			{
-				RightSample -= DelayLine.Num();
-			}
-			
-			const float Frac = SampleIndex - LeftSample;
-			OutSamples[OutputIndex] = (Frac * DelayLine[LeftSample]) + ((1.f - Frac) * DelayLine[RightSample]);
-			
-			SampleIndex += SampleStride;
-		}
-	}
-
-	int32 FInterpolatedMultiTapDelay::AlignIndex(const int32 InIndex) const
-	{
-		return InIndex - (InIndex % AUDIO_NUM_FLOATS_PER_VECTOR_REGISTER);
+		return DelayLine.GetAllocatedSize() > 0;
 	}
 }
\ No newline at end of file
diff --git a/Engine/Source/Runtime/SignalProcessing/Public/DSP/InterpolatedMultiTapDelay.h b/Engine/Source/Runtime/SignalProcessing/Public/DSP/InterpolatedMultiTapDelay.h
--- a/Engine/Source/Runtime/SignalProcessing/Public/DSP/InterpolatedMultiTapDelay.h
+++ b/Engine/Source/Runtime/SignalProcessing/Public/DSP/InterpolatedMultiTapDelay.h
@@ -3,6 +3,7 @@
 #pragma once
 
 #include "AlignedBuffer.h"
+#include "VectorLinearResampler.h"
 
 namespace Audio
 {
@@ -16,21 +17,22 @@ namespace Audio
 	public:
 		FInterpolatedMultiTapDelay() = default;
 		
-		SIGNALPROCESSING_API void Init(const int32 InBufferSizeSamples, const float InSampleRate);
+		SIGNALPROCESSING_API void Init(const int32 InBufferSizeSamples);
 
-		SIGNALPROCESSING_API void Advance(const FAlignedFloatBuffer& InSamples);
-		SIGNALPROCESSING_API void Read(const float StartDelayMSec, const float EndDelayMSec, FAlignedFloatBuffer& OutSamples);
+		SIGNALPROCESSING_API void Advance(const FAlignedFloatBuffer& InBuffer);
+
+		// Read and interpolate a variable number of samples back in the delay line
+		// return a fixed-point representation of fraction of the last sample, which should be passed as StartSampleFraction on successive reads for that tap.
+		SIGNALPROCESSING_API uint32 Read(const uint32 StartNumDelaySamples, const uint32 StartSampleFraction, const uint32 EndNumDelaySamples, FAlignedFloatBuffer& OutBuffer);
 		SIGNALPROCESSING_API void Reset();
 		SIGNALPROCESSING_API bool IsInitialized() const;
 
 	private:
-		void ReadBlockInternal(const int32 StartSample, const int32 SamplesToRead, const int32 NumOutputSamples, float* OutSamples);
-		FORCEINLINE int32 AlignIndex(const int32 InIndex) const;
-		
 		int32 WriteIndex = 0;
-		float MsToSamples = 1.f;
 
 		FAlignedFloatBuffer DelayLine;
+		FAlignedFloatBuffer WrapBuffer;
+		FVectorLinearResampler Resampler;
 	};	
 }
 
