Clamp linear resampler to ensure numerical stability.
#rnx

diff --git a/Engine/Plugins/Runtime/Metasound/Source/MetasoundEngine/Private/MetasoundWavePlayerNode.cpp b/Engine/Plugins/Runtime/Metasound/Source/MetasoundEngine/Private/MetasoundWavePlayerNode.cpp
--- a/Engine/Plugins/Runtime/Metasound/Source/MetasoundEngine/Private/MetasoundWavePlayerNode.cpp
+++ b/Engine/Plugins/Runtime/Metasound/Source/MetasoundEngine/Private/MetasoundWavePlayerNode.cpp
@@ -60,6 +60,8 @@ namespace Metasound
 
 	namespace WavePlayerNodePrivate
 	{
+		constexpr int32 MaxNumFramesToInterpolateFrameRatio = 128;
+
 		int32 GetCuePointFrame(const FSoundWaveCuePoint& InPoint)
 		{
 			return InPoint.FramePosition;
@@ -412,6 +414,8 @@ namespace Metasound
 
 		void Execute()
 		{
+			using namespace WavePlayerNodePrivate;
+
 			METASOUND_TRACE_CPUPROFILER_EVENT_SCOPE(Metasound::FWavePlayerOperator::Execute);
 
 			// Advance all triggers owned by this operator. 
@@ -432,7 +436,7 @@ namespace Metasound
 			// Update resampler with new frame ratio. 
 			if (Resampler.IsValid())
 			{
-				Resampler->SetFrameRatio(GetFrameRatio(), OperatorSettings.GetNumFramesPerBlock());
+				Resampler->SetFrameRatio(GetFrameRatio(), FMath::Min(OperatorSettings.GetNumFramesPerBlock(), MaxNumFramesToInterpolateFrameRatio));
 			}
 
 			// zero output buffers
@@ -648,7 +652,8 @@ namespace Metasound
 
 		float GetFrameRatio() const
 		{
-			return GetSampleRateFrameRatio() * GetPitchShiftFrameRatio();
+			using namespace Audio;
+			return FMath::Clamp(GetSampleRateFrameRatio() * GetPitchShiftFrameRatio(), FMultichannelLinearResampler::MinFrameRatio, FMultichannelLinearResampler::MaxFrameRatio);
 		}
 
 		float GetMaxPitchShiftFrameRatio() const
diff --git a/Engine/Source/Runtime/SignalProcessing/Private/MultichannelLinearResampler.cpp b/Engine/Source/Runtime/SignalProcessing/Private/MultichannelLinearResampler.cpp
--- a/Engine/Source/Runtime/SignalProcessing/Private/MultichannelLinearResampler.cpp
+++ b/Engine/Source/Runtime/SignalProcessing/Private/MultichannelLinearResampler.cpp
@@ -10,6 +10,9 @@
 
 namespace Audio
 {
+	const float FMultichannelLinearResampler::MaxFrameRatio = 100.f;
+	const float FMultichannelLinearResampler::MinFrameRatio = 0.001f;
+
 	FMultichannelLinearResampler::FMultichannelLinearResampler(int32 InNumChannels)
 		: NumChannels(InNumChannels)
 	{
@@ -17,7 +20,7 @@ namespace Audio
 
 	void FMultichannelLinearResampler::SetFrameRatio(float InRatio, int32 InDesiredNumFramesToInterpolate)
 	{
-		if (ensureMsgf(InRatio > 0.f, TEXT("The frame ratio must be greater than zero.")))
+		if (ensureMsgf((InRatio >= MinFrameRatio) && (InRatio <= MaxFrameRatio), TEXT("The frame ratio (%f) must be between %f and %f."), InRatio, MinFrameRatio, MaxFrameRatio))
 		{
 			if ((InDesiredNumFramesToInterpolate <= 1) || FMath::IsNearlyEqual(InRatio, CurrentFrameRatio))
 			{
@@ -199,7 +202,9 @@ namespace Audio
 			
 			const int32 NumBufferFrames = GetNumBufferFramesToProduceOutputFrames(NumOutputFrames);
 			NumOutputFrames = FMath::FloorToInt((NumAvailableInputFrames - NumBufferFrames) / FMath::Max(CurrentFrameRatio, TargetFrameRatio)) - 1;
+			NumOutputFrames = FMath::Max(NumOutputFrames, 0);
 			NumInputFramesRequired = NumAvailableInputFrames;
+			checkf(NumInputFramesRequired >= FMath::CeilToInt(MapOutputFrameToInputFrame(NumOutputFrames)), TEXT("Invalid calculation. Required input frames (%d) does not satisfy need for input frames (%f)"), InAudio.Num(), MapOutputFrameToInputFrame(NumOutputFrames));
 		}
 
 		if (NumOutputFrames > 0)
@@ -264,7 +269,7 @@ namespace Audio
 		{
 			return 0.f;
 		}
-		checkf(InAudio.Num() >= GetNumInputFramesNeededToProduceOutputFrames(NumOutputFrames), TEXT("Not enough input frames available to produce output frames"));
+		checkf(InAudio.Num() >= FMath::CeilToInt(MapOutputFrameToInputFrame(NumOutputFrames)), TEXT("Not enough input frames (%d) available to meet required input frames (%f)"), InAudio.Num(), MapOutputFrameToInputFrame(NumOutputFrames));
 
 		float* OutAudioData = OutAudio.GetData();
 		const float* InAudioData = InAudio.GetData();
diff --git a/Engine/Source/Runtime/SignalProcessing/Public/DSP/MultichannelLinearResampler.h b/Engine/Source/Runtime/SignalProcessing/Public/DSP/MultichannelLinearResampler.h
--- a/Engine/Source/Runtime/SignalProcessing/Public/DSP/MultichannelLinearResampler.h
+++ b/Engine/Source/Runtime/SignalProcessing/Public/DSP/MultichannelLinearResampler.h
@@ -13,6 +13,10 @@ namespace Audio
 	{
 		static constexpr float MinFrameRatioFrameDelta = 0.001f;
 	public:
+
+		SIGNALPROCESSING_API static const float MaxFrameRatio;
+		SIGNALPROCESSING_API static const float MinFrameRatio;
+
 		/** Construct a linear resampler.
 		 *
 		 * @param InNumChannel - Number of audio channels in input and output buffers.
