Fix for crash in multichannel linear resampler.
#rnx

diff --git a/Engine/Source/Runtime/SignalProcessing/Private/MultichannelLinearResampler.cpp b/Engine/Source/Runtime/SignalProcessing/Private/MultichannelLinearResampler.cpp
--- a/Engine/Source/Runtime/SignalProcessing/Private/MultichannelLinearResampler.cpp
+++ b/Engine/Source/Runtime/SignalProcessing/Private/MultichannelLinearResampler.cpp
@@ -204,7 +204,7 @@ namespace Audio
 			NumOutputFrames = FMath::FloorToInt((NumAvailableInputFrames - NumBufferFrames) / FMath::Max(CurrentFrameRatio, TargetFrameRatio)) - 1;
 			NumOutputFrames = FMath::Max(NumOutputFrames, 0);
 			NumInputFramesRequired = NumAvailableInputFrames;
-			checkf(NumInputFramesRequired >= FMath::CeilToInt(MapOutputFrameToInputFrame(NumOutputFrames)), TEXT("Invalid calculation. Required input frames (%d) does not satisfy need for input frames (%f)"), InAudio.Num(), MapOutputFrameToInputFrame(NumOutputFrames));
+			checkf(NumInputFramesRequired > FMath::CeilToInt(MapOutputFrameToInputFrame(NumOutputFrames - 1)), TEXT("Invalid calculation. Required input frames (%d) does not satisfy need for input frames (%f)"), NumInputFramesRequired, MapOutputFrameToInputFrame(NumOutputFrames - 1));
 		}
 
 		if (NumOutputFrames > 0)
@@ -269,7 +269,6 @@ namespace Audio
 		{
 			return 0.f;
 		}
-		checkf(InAudio.Num() >= FMath::CeilToInt(MapOutputFrameToInputFrame(NumOutputFrames)), TEXT("Not enough input frames (%d) available to meet required input frames (%f)"), InAudio.Num(), MapOutputFrameToInputFrame(NumOutputFrames));
 
 		float* OutAudioData = OutAudio.GetData();
 		const float* InAudioData = InAudio.GetData();
@@ -315,7 +314,7 @@ namespace Audio
 			}
 
 			// Check for buffer over run
-			checkf((LowerFrameIndex + 1) < InAudio.Num(), TEXT("Buffer overrun in multichannel linear resampler. Attempt to read index %d of array with %d elements."), LowerFrameIndex + 1, InAudio.Num());
+			checkf((LowerFrameIndex + 1) < InAudio.Num(), TEXT("Buffer overrun in multichannel linear resampler. Attempt to read index %d of array with %d elements. FrameRatio: %f, FrameRatioDelta: %f, NumFramesToInterpolate: %d"), LowerFrameIndex + 1, InAudio.Num(), CurrentFrameRatio, FrameRatioFrameDelta, NumFramesToInterpolate);
 
 			return InputFrameIndex;
 		}
