defend against overrunning internal wrap buffer in interpolated multi-tap delay

#rnx

diff --git a/Engine/Source/Runtime/SignalProcessing/Private/InterpolatedMultiTapDelay.cpp b/Engine/Source/Runtime/SignalProcessing/Private/InterpolatedMultiTapDelay.cpp
--- a/Engine/Source/Runtime/SignalProcessing/Private/InterpolatedMultiTapDelay.cpp
+++ b/Engine/Source/Runtime/SignalProcessing/Private/InterpolatedMultiTapDelay.cpp
@@ -3,21 +3,22 @@
 
 #include "DSP/InterpolatedMultiTapDelay.h"
 
+#include "SignalProcessingModule.h"
 #include "DSP/FloatArrayMath.h"
 #include "Math/VectorRegister.h"
 
 namespace Audio
 {
-	void FInterpolatedMultiTapDelay::Init(const int32 InBufferSizeSamples)
+	void FInterpolatedMultiTapDelay::Init(const int32 InDelayBufferSamples)
 	{
 		WriteIndex = 0;
 		DelayLine.Reset();
-		DelayLine.AddZeroed(InBufferSizeSamples);
+		DelayLine.AddZeroed(InDelayBufferSamples);
 		WrapBuffer.Reset();
-		WrapBuffer.AddZeroed(InBufferSizeSamples);
+		WrapBuffer.AddZeroed(InDelayBufferSamples);
 	}
 
-	void FInterpolatedMultiTapDelay::Advance(const FAlignedFloatBuffer& InBuffer)
+	void FInterpolatedMultiTapDelay::Advance(TArrayView<const float> InBuffer)
 	{
 		check (InBuffer.Num() % AUDIO_NUM_FLOATS_PER_VECTOR_REGISTER == 0);
 		const uint32 InNumSamples = InBuffer.Num();
@@ -51,20 +52,18 @@ namespace Audio
 		}
 	}
 
-	uint32 FInterpolatedMultiTapDelay::Read(const uint32 StartNumDelaySamples, const uint32 StartSampleFraction, const uint32 EndNumDelaySamples, FAlignedFloatBuffer& OutBuffer)
+	uint32 FInterpolatedMultiTapDelay::Read(const uint32 StartNumDelaySamples, const uint32 StartSampleFraction, const uint32 EndNumDelaySamples, TArrayView<float> OutBuffer)
 	{
 		const int32 OutputNumSamples = OutBuffer.Num();
-		const int32 DelayBufferNumSamples = DelayLine.Num();
+		int32 DelayBufferNumSamples = DelayLine.Num();
 
 		// likely to only run on the first frame, if not configured with enough memory
 		if (OutputNumSamples > DelayBufferNumSamples)
 		{
-			DelayLine.SetNumZeroed(OutputNumSamples);
-		}
+			DelayBufferNumSamples = 2 * OutputNumSamples + 1;
+			DelayLine.SetNumZeroed(DelayBufferNumSamples);
 
-		if (OutputNumSamples > WrapBuffer.Num())
-		{
-			WrapBuffer.SetNumUninitialized(OutputNumSamples);
+			UE_LOG(LogSignalProcessing, Warning, TEXT("FInterpolatedMultiTapDelay not configured with enough memory to process an output buffer - allocating extra space."));
 		}
 		
 		if (DelayBufferNumSamples <= 0 || OutputNumSamples <= 0)
@@ -72,8 +71,8 @@ namespace Audio
 			return 0;
 		}
 
-		int32 StartSample = WriteIndex - StartNumDelaySamples;
-		int32 EndSample = WriteIndex + OutputNumSamples - EndNumDelaySamples;
+		int32 StartSample = WriteIndex - OutputNumSamples - FMath::Clamp(StartNumDelaySamples, 0, DelayBufferNumSamples - 1);
+		int32 EndSample = WriteIndex - FMath::Clamp(EndNumDelaySamples, 0, DelayBufferNumSamples - 1);
 
 		const float SampleStride = FMath::Clamp((float)(EndSample - StartSample) / (float)OutputNumSamples, 0.25f, 4.f);
 		const uint32 FixedSampleRate = (uint32)(SampleStride * 65536.f);
@@ -86,14 +85,26 @@ namespace Audio
 		float* SourceBuffer = &DelayLine[StartSample];
 		if (StartSample + FramesNeeded >= DelayBufferNumSamples)
 		{
-			constexpr int32 NumSafetyBufferSamples = 16;
-			
 			const int32 NumSamplesToEnd = DelayBufferNumSamples - StartSample;
-			const int32 SecondBufferNumSamples = FMath::Min((FramesNeeded - NumSamplesToEnd) + NumSafetyBufferSamples, DelayBufferNumSamples);
+			const int32 SecondBufferNumSamples = FramesNeeded - NumSamplesToEnd;
+			if (FramesNeeded > WrapBuffer.Num())
+			{
+				WrapBuffer.SetNumUninitialized(FramesNeeded);
+			}
 			
 			FMemory::Memcpy(WrapBuffer.GetData(), &DelayLine[StartSample], NumSamplesToEnd * sizeof(float));
-			FMemory::Memcpy(&WrapBuffer[NumSamplesToEnd], DelayLine.GetData(), SecondBufferNumSamples * sizeof(float));
-
+			// if used sensibly this shouldn't happen, but better to inject 0's than to over-read the buffer
+			if (SecondBufferNumSamples > StartSample)
+			{
+				FMemory::Memcpy(&WrapBuffer[NumSamplesToEnd], DelayLine.GetData(), StartSample * sizeof(float));
+				
+				const int32 FramesRemaining = SecondBufferNumSamples - StartSample;
+				FMemory::Memzero(&WrapBuffer[NumSamplesToEnd + SecondBufferNumSamples], FramesRemaining * sizeof(float));
+			}
+			else
+			{
+				FMemory::Memcpy(&WrapBuffer[NumSamplesToEnd], DelayLine.GetData(), SecondBufferNumSamples * sizeof(float));
+			}
 			SourceBuffer = WrapBuffer.GetData();
 		}
 
diff --git a/Engine/Source/Runtime/SignalProcessing/Public/DSP/InterpolatedMultiTapDelay.h b/Engine/Source/Runtime/SignalProcessing/Public/DSP/InterpolatedMultiTapDelay.h
--- a/Engine/Source/Runtime/SignalProcessing/Public/DSP/InterpolatedMultiTapDelay.h
+++ b/Engine/Source/Runtime/SignalProcessing/Public/DSP/InterpolatedMultiTapDelay.h
@@ -17,13 +17,13 @@ namespace Audio
 	public:
 		FInterpolatedMultiTapDelay() = default;
 		
-		SIGNALPROCESSING_API void Init(const int32 InBufferSizeSamples);
+		SIGNALPROCESSING_API void Init(const int32 InDelayBufferSamples);
 
-		SIGNALPROCESSING_API void Advance(const FAlignedFloatBuffer& InBuffer);
+		SIGNALPROCESSING_API void Advance(TArrayView<const float> InBuffer);
 
 		// Read and interpolate a variable number of samples back in the delay line
 		// return a fixed-point representation of fraction of the last sample, which should be passed as StartSampleFraction on successive reads for that tap.
-		SIGNALPROCESSING_API uint32 Read(const uint32 StartNumDelaySamples, const uint32 StartSampleFraction, const uint32 EndNumDelaySamples, FAlignedFloatBuffer& OutBuffer);
+		SIGNALPROCESSING_API uint32 Read(const uint32 StartNumDelaySamples, const uint32 StartSampleFraction, const uint32 EndNumDelaySamples, TArrayView<float> OutBuffer);
 		SIGNALPROCESSING_API void Reset();
 		SIGNALPROCESSING_API bool IsInitialized() const;
 
