[Metasound] fix filters clamping to bad numbers in lower samplerates

#rnx

diff --git a/Engine/Source/Runtime/SignalProcessing/Private/Filter.cpp b/Engine/Source/Runtime/SignalProcessing/Private/Filter.cpp
--- a/Engine/Source/Runtime/SignalProcessing/Private/Filter.cpp
+++ b/Engine/Source/Runtime/SignalProcessing/Private/Filter.cpp
@@ -7,7 +7,7 @@ namespace Audio
 {
     float FBiquadFilter::ClampCutoffFrequency(float InCutoffFrequency)
     {
-        return FMath::Clamp(InCutoffFrequency, 5.0f, SampleRate / 2.0f - 1.0f);
+        return FMath::Clamp(InCutoffFrequency, 5.0f, 0.9f * (SampleRate / 2.0f));
     }
     
 	FBiquadFilter::FBiquadFilter()
@@ -401,7 +401,7 @@ namespace Audio
 	{
 		if (BaseFrequency != InCutoffFrequency)
 		{
-			BaseFrequency = InCutoffFrequency;
+			BaseFrequency = FMath::Clamp(InCutoffFrequency, 20.f, 0.9f * (SampleRate / 2.f));
 			bChanged = true;
 		}
 	}
@@ -450,7 +450,7 @@ namespace Audio
 		{
 			bChanged = false;
 
-			Frequency = FMath::Clamp(BaseFrequency * GetFrequencyMultiplier(ModFrequency + ExternalModFrequency), 80.0f, 18000.0f);
+			Frequency = FMath::Clamp(BaseFrequency * GetFrequencyMultiplier(ModFrequency + ExternalModFrequency), 80.0f, 0.9f * (SampleRate / 2.f));
 			Q = BaseQ + ModQ + ExternalModQ;
 		}
 	}
