Optimize submix effect dynamics processor by utilizing deinterleaved codepath
#rnx

diff --git a/Engine/Plugins/Runtime/Metasound/Source/MetasoundStandardNodes/Private/MetasoundCompressorNode.cpp b/Engine/Plugins/Runtime/Metasound/Source/MetasoundStandardNodes/Private/MetasoundCompressorNode.cpp
--- a/Engine/Plugins/Runtime/Metasound/Source/MetasoundStandardNodes/Private/MetasoundCompressorNode.cpp
+++ b/Engine/Plugins/Runtime/Metasound/Source/MetasoundStandardNodes/Private/MetasoundCompressorNode.cpp
@@ -393,7 +393,7 @@ namespace Metasound
 			}
 
 			// Apply lookahead delay to dry signal
-			InputDelay.ProcessAudio(*AudioInput, DelayedInputSignal);
+			InputDelay.ProcessAudio(*AudioInput, TArrayView<float>(DelayedInputSignal));
 
 			const float* InSamples = AudioInput->GetData();
 			float* OutSamples      = AudioOutput->GetData();
diff --git a/Engine/Source/Runtime/AudioMixer/Classes/SubmixEffects/AudioMixerSubmixEffectDynamicsProcessor.h b/Engine/Source/Runtime/AudioMixer/Classes/SubmixEffects/AudioMixerSubmixEffectDynamicsProcessor.h
--- a/Engine/Source/Runtime/AudioMixer/Classes/SubmixEffects/AudioMixerSubmixEffectDynamicsProcessor.h
+++ b/Engine/Source/Runtime/AudioMixer/Classes/SubmixEffects/AudioMixerSubmixEffectDynamicsProcessor.h
@@ -298,9 +298,6 @@ protected:
 	
 	Audio::FAlignedFloatBuffer AudioExternal;
 
-	TArray<float> AudioKeyFrame;
-	TArray<float> AudioInputFrame;
-
 	Audio::FDeviceId DeviceId = INDEX_NONE;
 
 	bool bBypass = false;
diff --git a/Engine/Source/Runtime/AudioMixer/Private/Effects/AudioMixerSubmixEffectDynamicsProcessor.cpp b/Engine/Source/Runtime/AudioMixer/Private/Effects/AudioMixerSubmixEffectDynamicsProcessor.cpp
--- a/Engine/Source/Runtime/AudioMixer/Private/Effects/AudioMixerSubmixEffectDynamicsProcessor.cpp
+++ b/Engine/Source/Runtime/AudioMixer/Private/Effects/AudioMixerSubmixEffectDynamicsProcessor.cpp
@@ -7,6 +7,7 @@
 #include "AudioMixerDevice.h"
 #include "AudioMixerSubmix.h"
 #include "ProfilingDebugging/CsvProfiler.h"
+#include "ProfilingDebugging/CpuProfilerTrace.h"
 
 #include UE_INLINE_GENERATED_CPP_BY_NAME(AudioMixerSubmixEffectDynamicsProcessor)
 
@@ -47,12 +48,6 @@ void FSubmixEffectDynamicsProcessor::Init(const FSoundEffectSubmixInitData& Init
 
 	DynamicsProcessor.Init(InitData.SampleRate, ProcessorScratchNumChannels);
 
-	AudioKeyFrame.Reset();
-	AudioKeyFrame.AddZeroed(ProcessorScratchNumChannels);
-
-	AudioInputFrame.Reset();
-	AudioInputFrame.AddZeroed(ProcessorScratchNumChannels);
-
 	DeviceId = InitData.DeviceID;
 
 	if (USubmixEffectDynamicsProcessorPreset* ProcPreset = Cast<USubmixEffectDynamicsProcessorPreset>(Preset.Get()))
@@ -254,14 +249,14 @@ void FSubmixEffectDynamicsProcessor::OnProcessAudio(const FSoundEffectSubmixInpu
 {
 	CSV_SCOPED_TIMING_STAT(Audio, SubmixDynamics);
 	SCOPE_CYCLE_COUNTER(STAT_AudioMixerSubmixDynamics);
+	TRACE_CPUPROFILER_EVENT_SCOPE(FSubmixEffectDynamicsProcessor::OnProcessAudio);
 
 	ensure(InData.NumChannels == OutData.NumChannels);
 
 	const Audio::FAlignedFloatBuffer& InBuffer = *InData.AudioBuffer;
 	Audio::FAlignedFloatBuffer& OutBuffer = *OutData.AudioBuffer;
 
-	const bool bBypassDueToInvalidChannelCount = !ensure(InData.NumChannels <= AudioInputFrame.Num());
-	if (bBypassDueToInvalidChannelCount || bBypassSubmixDynamicsProcessor || bBypass)
+	if (bBypassSubmixDynamicsProcessor || bBypass)
 	{
 		FMemory::Memcpy(OutBuffer.GetData(), InBuffer.GetData(), sizeof(float) * InBuffer.Num());
 		return;
@@ -302,27 +297,12 @@ void FSubmixEffectDynamicsProcessor::OnProcessAudio(const FSoundEffectSubmixInpu
 	// No key assigned (Uses input buffer as key)
 	if (KeySource.GetType() == ESubmixEffectDynamicsKeySource::Default)
 	{
-		for (int32 Frame = 0; Frame < InData.NumFrames; ++Frame)
-		{
-			const int32 SampleIndex = Frame * InData.NumChannels;
-			DynamicsProcessor.ProcessAudio(&InBuffer[SampleIndex], InData.NumChannels, &OutBuffer[SampleIndex]);
-		}
+		DynamicsProcessor.ProcessAudio(InBuffer.GetData(), InData.NumChannels * InData.NumFrames, OutBuffer.GetData());
 	}
 	// Key assigned
 	else
 	{
-		for (int32 Frame = 0; Frame < InData.NumFrames; ++Frame)
-		{
-			// Copy the data to the input frame
-			const int32 SampleIndexOfInputFrame = Frame * InData.NumChannels;
-			FMemory::Memcpy(AudioInputFrame.GetData(), &InBuffer[SampleIndexOfInputFrame], sizeof(float) * InData.NumChannels);
-
-			// Copy the data to the key frame
-			const int32 SampleIndexOfKeyFrame = Frame * NumKeyChannels;
-			FMemory::Memcpy(AudioKeyFrame.GetData(), &AudioExternal[SampleIndexOfKeyFrame], sizeof(float) * NumKeyChannels);
-
-			DynamicsProcessor.ProcessAudio(AudioInputFrame.GetData(), InData.NumChannels, &OutBuffer[SampleIndexOfInputFrame], AudioKeyFrame.GetData());
-		}
+		DynamicsProcessor.ProcessAudio(InBuffer.GetData(), InData.NumChannels * InData.NumFrames, OutBuffer.GetData(), AudioExternal.GetData());
 	}
 }
 
diff --git a/Engine/Source/Runtime/SignalProcessing/Private/DynamicsProcessor.cpp b/Engine/Source/Runtime/SignalProcessing/Private/DynamicsProcessor.cpp
--- a/Engine/Source/Runtime/SignalProcessing/Private/DynamicsProcessor.cpp
+++ b/Engine/Source/Runtime/SignalProcessing/Private/DynamicsProcessor.cpp
@@ -2,6 +2,7 @@
 
 #include "DSP/DynamicsProcessor.h"
 #include "DSP/FloatArrayMath.h"
+#include "ProfilingDebugging/CpuProfilerTrace.h"
 #include "SignalProcessingModule.h"
 
 namespace Audio
@@ -10,7 +11,7 @@ namespace Audio
 		: ProcessingMode(EDynamicsProcessingMode::Compressor)
 		, SlopeFactor(0.0f)
 		, EnvelopeFollowerPeakMode(EPeakMode::Peak)
-		, LookaheedDelayMsec(10.0f)
+		, LookaheadDelayMsec(10.0f)
 		, AttackTimeMsec(20.0f)
 		, ReleaseTimeMsec(1000.0f)
 		, ThresholdDb(-6.0f)
@@ -19,6 +20,7 @@ namespace Audio
 		, InputGain(1.0f)
 		, OutputGain(1.0f)
 		, KeyGain(1.0f)
+		, SampleRate(48000.f)
 		, LinkMode(EDynamicsProcessorChannelLinkMode::Disabled)
 		, bIsAnalogMode(true)
 		, bKeyAuditionEnabled(false)
@@ -36,22 +38,18 @@ namespace Audio
 
 	void FDynamicsProcessor::Init(const float InSampleRate, const int32 InNumChannels)
 	{
+		check(InSampleRate > 0.f);
 		SampleRate = InSampleRate;
 
 		SetNumChannels(InNumChannels);
 		SetKeyNumChannels(InNumChannels);
-	
-		LookaheadDelay.Reset();
-		LookaheadDelay.AddDefaulted(InNumChannels);
 
 		for (int32 Channel = 0; Channel < InNumChannels; ++Channel)
 		{
-			LookaheadDelay[Channel].Init(SampleRate, MaxLookaheadMsec / 1000.0f);
-			LookaheadDelay[Channel].SetDelayMsec(LookaheedDelayMsec);
-
 			EnvFollower[Channel].Init(FInlineEnvelopeFollowerInitParams{SampleRate, AttackTimeMsec, ReleaseTimeMsec, EnvelopeFollowerPeakMode, bIsAnalogMode});
 		}
 
+		// Initialize key filters
 		InputLowshelfFilter.Init(SampleRate, InNumChannels, EBiquadFilter::LowShelf);
 		InputHighshelfFilter.Init(SampleRate, InNumChannels, EBiquadFilter::HighShelf);
 
@@ -79,10 +77,11 @@ namespace Audio
 	
 	void FDynamicsProcessor::SetLookaheadMsec(const float InLookAheadMsec)
 	{
-		LookaheedDelayMsec = InLookAheadMsec;
+		LookaheadDelayMsec = InLookAheadMsec;
+		const int32 NumDelayFrames = GetNumDelayFrames();
 		for (int32 Channel = 0; Channel < LookaheadDelay.Num(); ++Channel)
 		{
-			LookaheadDelay[Channel].SetDelayMsec(LookaheedDelayMsec);
+			LookaheadDelay[Channel].SetDelayLengthSamples(NumDelayFrames);
 		}
 	}
 
@@ -226,13 +225,16 @@ namespace Audio
 
 		if (InNumChannels != LookaheadDelay.Num())
 		{
+			// Initialize lookahead buffers
+			constexpr int32 DelayInternalBufferSize = 32;
+			const int32 MaxNumDelayFrames = FMath::Max(1, FMath::CeilToInt(MaxLookaheadMsec * SampleRate / 1000.0f));
+			const int32 NumDelayFrames = GetNumDelayFrames();
+
 			LookaheadDelay.Reset();
-			LookaheadDelay.AddDefaulted(InNumChannels);
 
 			for (int32 Channel = 0; Channel < InNumChannels; ++Channel)
 			{
-				LookaheadDelay[Channel].Init(SampleRate, 0.1f);
-				LookaheadDelay[Channel].SetDelayMsec(LookaheedDelayMsec);
+				LookaheadDelay.Emplace(MaxNumDelayFrames, NumDelayFrames, DelayInternalBufferSize);
 			}
 		}
 	}
@@ -294,65 +296,93 @@ namespace Audio
 		}
 	}
 
-	void FDynamicsProcessor::ProcessAudio(const float* InBuffer, const int32 InNumSamples, float* OutBuffer, const float* InKeyBuffer)
+
+	void FDynamicsProcessor::ProcessAudio(const float* InBuffer, const int32 InNumSamples, float* OutBuffer, const float* InKeyBuffer, float* OutEnvelope)
 	{
-		check(nullptr != InBuffer);
-		check(nullptr != OutBuffer);
+		TRACE_CPUPROFILER_EVENT_SCOPE(FDynamicsProcessor::ProcessAudio_Interleaved);
+
+		constexpr int32 MaxNumChannels = 8;
+		constexpr int32 MinNumSubbufferFrames = 32;
+		constexpr int32 NumSubbufferSamples = MaxNumChannels * MinNumSubbufferFrames;
+		static_assert((NumSubbufferSamples * sizeof(float) * 4) <= 4096, "Statically allocated deinterleave buffers are clamped to be 4096 to protect against running out of stack memory.");
+
+		// Stack allocated intermediate buffers to hold deinterleave audio
+		float* DeinterleaveInput[MaxNumChannels];
+		float InputWorkBuffer[NumSubbufferSamples];
+
+		float* DeinterleaveOutput[MaxNumChannels];
+		float OutputWorkBuffer[NumSubbufferSamples];
+
+		float* DeinterleaveKey[MaxNumChannels];
+		float KeyWorkBuffer[NumSubbufferSamples];
+
+		float* DeinterleaveEnvelope[MaxNumChannels];
+		float EnvelopeWorkBuffer[NumSubbufferSamples];
+
 
 		const int32 NumChannels = GetNumChannels();
 		const int32 KeyNumChannels = GetKeyNumChannels();
+		const int32 NumInputFrames = InNumSamples / NumChannels;
+		const int32 NumSubbufferFrames = NumSubbufferSamples / FMath::Max(1, FMath::Max(NumChannels, KeyNumChannels));
 
-		if (InKeyBuffer)
+		check(nullptr != InBuffer);
+		check(nullptr != OutBuffer);
+		check(KeyNumChannels <= MaxNumChannels);
+		check(NumChannels <= MaxNumChannels);
+		check(NumSubbufferFrames >= MinNumSubbufferFrames);
+
+		// Initialize deinterleave buffer pointers
+		for (int32 ChannelIndex = 0; ChannelIndex < NumChannels; ChannelIndex++)
 		{
-			int32 KeySampleIndex = 0;
-			for (int32 SampleIndex = 0; SampleIndex < InNumSamples; SampleIndex += NumChannels)
-			{
-				const float* KeyFrame = &InKeyBuffer[KeySampleIndex];
-				ProcessAudioFrame(&InBuffer[SampleIndex], &OutBuffer[SampleIndex], KeyFrame);
-				KeySampleIndex += KeyNumChannels;
-			}
-		}
-		else
+			DeinterleaveInput[ChannelIndex] = &InputWorkBuffer[ChannelIndex * NumSubbufferFrames];
+			DeinterleaveOutput[ChannelIndex] = &OutputWorkBuffer[ChannelIndex * NumSubbufferFrames];
+			DeinterleaveEnvelope[ChannelIndex] = &EnvelopeWorkBuffer[ChannelIndex * NumSubbufferFrames];
+		};
+		for (int32 KeyChannelIndex = 0; KeyChannelIndex < NumChannels; KeyChannelIndex++)
 		{
-			for (int32 SampleIndex = 0; SampleIndex < InNumSamples; SampleIndex += NumChannels)
-			{
-				const float* KeyFrame = &InBuffer[SampleIndex];
-				ProcessAudioFrame(&InBuffer[SampleIndex], &OutBuffer[SampleIndex], KeyFrame);
-			}
+			DeinterleaveKey[KeyChannelIndex] = &KeyWorkBuffer[KeyChannelIndex * NumSubbufferFrames];
 		}
-	}
 
-	void FDynamicsProcessor::ProcessAudio(const float* InBuffer, const int32 InNumSamples, float* OutBuffer, const float* InKeyBuffer, float* OutEnvelope)
-	{
-		check(nullptr != InBuffer);
-		check(nullptr != OutBuffer);
-		check(nullptr != OutEnvelope);
+		// Run dyanmics processor on deinterleaved subbuffers. 
+		for (int32 FrameIndex = 0; FrameIndex < NumInputFrames; FrameIndex += NumSubbufferFrames)
+		{
+			const int32 NumFramesThisIteration = FMath::Min(NumSubbufferFrames, NumInputFrames - FrameIndex);
+			const int32 InterleavedSampleIndex = FrameIndex * NumChannels;
 
-		const int32 NumChannels = GetNumChannels();
-		const int32 KeyNumChannels = GetKeyNumChannels();
+			ArrayDeinterleave(&InBuffer[InterleavedSampleIndex], DeinterleaveInput, NumFramesThisIteration, NumChannels);
 
-		if (InKeyBuffer)
-		{
-			int32 KeySampleIndex = 0;
-			for (int32 SampleIndex = 0; SampleIndex < InNumSamples; SampleIndex += NumChannels)
+			if (InKeyBuffer)
 			{
-				const float* KeyFrame = &InKeyBuffer[KeySampleIndex];
-				ProcessAudioFrame(&InBuffer[SampleIndex], &OutBuffer[SampleIndex], KeyFrame, &OutEnvelope[SampleIndex]);
-				KeySampleIndex += KeyNumChannels;
+				const int32 KeyInterleavedSampleIndex = FrameIndex * KeyNumChannels;
+				ArrayDeinterleave(&InKeyBuffer[KeyInterleavedSampleIndex], DeinterleaveKey, NumFramesThisIteration, KeyNumChannels);
+				
+				ProcessAudio(DeinterleaveInput, NumFramesThisIteration, DeinterleaveOutput, DeinterleaveKey, DeinterleaveEnvelope);
 			}
-		}
-		else
-		{
-			for (int32 SampleIndex = 0; SampleIndex < InNumSamples; SampleIndex += NumChannels)
+			else
+			{
+				ProcessAudio(DeinterleaveInput, NumFramesThisIteration, DeinterleaveOutput, nullptr /* InKeyBuffers */, DeinterleaveEnvelope);
+			}
+
+			// Interleave output results
+			if (OutEnvelope)
 			{
-				const float* KeyFrame = &InBuffer[SampleIndex];
-				ProcessAudioFrame(&InBuffer[SampleIndex], &OutBuffer[SampleIndex], KeyFrame, &OutEnvelope[SampleIndex]);
+				ArrayInterleave(DeinterleaveEnvelope, &OutEnvelope[InterleavedSampleIndex], NumFramesThisIteration, NumChannels);
 			}
+
+			ArrayInterleave(DeinterleaveOutput, &OutBuffer[InterleavedSampleIndex], NumFramesThisIteration, NumChannels);
 		}
 	}
 
-	void FDynamicsProcessor::ProcessAudio(const float* const* const InBuffers, const int32 InNumSamples, float* const* OutBuffers, const float* const* const InKeyBuffers, float* const* OutEnvelopes)
+	void FDynamicsProcessor::ProcessAudio(const float* const* const InBuffers, const int32 InNumFrames, float* const* OutBuffers, const float* const* const InKeyBuffers, float* const* OutEnvelopes)
 	{
+		// we need this test because we are going to use the output buffer as our scratch pad for key processing!
+		check(GetKeyNumChannels() <= GetNumChannels()); 
+		ProcessAudioDeinterleaveInternal(InBuffers, InNumFrames, OutBuffers, InKeyBuffers, OutEnvelopes);
+	}
+
+	void FDynamicsProcessor::ProcessAudioDeinterleaveInternal(const float* const* const InBuffers, const int32 InNumFrames, float* const* OutBuffers, const float* const* const InKeyBuffers, float* const* OutEnvelopes)
+	{
+		TRACE_CPUPROFILER_EVENT_SCOPE(FDynamicsProcessor::ProcessAudioDeinterleaveInternal);
 		// NOTE: This is a useful test, but it isn't comprehensive. It can't be sure the passed 
 		// in arrays of pointers have a matching number of pointers to the channel counts. It
 		// will have to do for now.
@@ -360,25 +390,24 @@ namespace Audio
 		check(nullptr != OutBuffers && nullptr != OutBuffers[0]);
 		check(nullptr != OutEnvelopes && nullptr != OutEnvelopes[0]);
 
+		// NOTE: The OutBuffers are used as scratch buffers for any processing applied
+		// to the key. Callers must provide 
 		const int32 KeyNumChannels = GetKeyNumChannels();
 		const int32 NumChannels = GetNumChannels();
 
-		// we need this test because we are going to use the output buffer as our scratch pad for key processing!
-		check(KeyNumChannels <= NumChannels);
-
 		const float* const* KeySources = InKeyBuffers ? InKeyBuffers : InBuffers;
 		float* const* KeyWorkBuffers = OutBuffers;
 
 		// Process Key
 		if (bKeyLowshelfEnabled)
 		{
-			InputLowshelfFilter.ProcessAudio(KeySources, InNumSamples, KeyWorkBuffers);
+			InputLowshelfFilter.ProcessAudio(KeySources, InNumFrames, KeyWorkBuffers);
 			KeySources = KeyWorkBuffers;
 		}
 
 		if (bKeyHighshelfEnabled)
 		{
-			InputHighshelfFilter.ProcessAudio(KeySources, InNumSamples, KeyWorkBuffers);
+			InputHighshelfFilter.ProcessAudio(KeySources, InNumFrames, KeyWorkBuffers);
 			KeySources = KeyWorkBuffers;
 		}
 
@@ -399,13 +428,13 @@ namespace Audio
 				// "moved" the key input to the output buffer. We have to do it here. 
 				for (int32 Channel = 0; Channel < KeyNumChannels; ++Channel)
 				{
-					FMemory::Memcpy(OutBuffers[Channel], KeySources[Channel], sizeof(float) * InNumSamples);
+					FMemory::Memcpy(OutBuffers[Channel], KeySources[Channel], sizeof(float) * InNumFrames);
 				}
 			}
 			// we just have to zero out any channels that are not in the key signal
 			for (int32 Channel = KeyNumChannels; Channel < NumChannels; ++Channel)
 			{
-				FMemory::Memset(OutBuffers[Channel], 0, sizeof(float) * InNumSamples);
+				FMemory::Memset(OutBuffers[Channel], 0, sizeof(float) * InNumFrames);
 			}
 
 			// I WOULD zero or 1 the "OutEnvelope", but that is not what the original does!
@@ -415,9 +444,9 @@ namespace Audio
 		for (int32 Channel = 0; Channel < KeyNumChannels; ++Channel)
 		{
 			// apply detector gain to key
-			Audio::ArrayMultiplyByConstant(TArrayView<const float>(KeySources[Channel], InNumSamples), DetectorGain, TArrayView<float>(KeyWorkBuffers[Channel], InNumSamples));
+			Audio::ArrayMultiplyByConstant(TArrayView<const float>(KeySources[Channel], InNumFrames), DetectorGain, TArrayView<float>(KeyWorkBuffers[Channel], InNumFrames));
 			// EnvelopeFollow key
-			EnvFollower[Channel].ProcessBuffer(KeyWorkBuffers[Channel], InNumSamples, OutEnvelopes[Channel]);
+			EnvFollower[Channel].ProcessBuffer(KeyWorkBuffers[Channel], InNumFrames, OutEnvelopes[Channel]);
 		}
 
 		int32 NumGainChannels = 1;
@@ -425,33 +454,35 @@ namespace Audio
 		{
 		case EDynamicsProcessorChannelLinkMode::Average:
 		{
+			TRACE_CPUPROFILER_EVENT_SCOPE(FDynamicsProcessor::ProcessAudioDeinterleaveInternal_LinkModeAverage);
 			// average all key channels
-			TArrayView<float> Out = TArrayView<float>(OutEnvelopes[0], InNumSamples);
+			TArrayView<float> Out = TArrayView<float>(OutEnvelopes[0], InNumFrames);
 			for (int32 KeyChannel = 1; KeyChannel < KeyNumChannels; ++KeyChannel)
 			{
-				ArraySum(TArrayView<const float>(OutEnvelopes[KeyChannel], InNumSamples), Out, Out);
+				ArraySum(TArrayView<const float>(OutEnvelopes[KeyChannel], InNumFrames), Out, Out);
 			}
 			ArrayMultiplyByConstant(Out, 1.0f / (float)KeyNumChannels, Out);
 			// convert magnitude to db
 			ArrayMagnitudeToDecibelInPlace(Out, -96.0f);
 			// compute 1 gain and use for all channels ...	const float ComputedGain = ComputeGain(DetectorOutLinkedDb);
-			ComputeGains(OutEnvelopes[0], InNumSamples);
+			ComputeGains(OutEnvelopes[0], InNumFrames);
 			NumGainChannels = 1;
 		}
 		break;
 
 		case EDynamicsProcessorChannelLinkMode::Peak:
 		{
+			TRACE_CPUPROFILER_EVENT_SCOPE(FDynamicsProcessor::ProcessAudioDeinterleaveInternal_LinkModePeak);
 			// detect peak across all key channels
-			TArrayView<float> Out = TArrayView<float>(OutEnvelopes[0], InNumSamples);
+			TArrayView<float> Out = TArrayView<float>(OutEnvelopes[0], InNumFrames);
 			for (int32 KeyChannel = 1; KeyChannel < KeyNumChannels; ++KeyChannel)
 			{
-				ArrayMax(TArrayView<const float>(OutEnvelopes[KeyChannel], InNumSamples), Out, Out);
+				ArrayMax(TArrayView<const float>(OutEnvelopes[KeyChannel], InNumFrames), Out, Out);
 			}
 			// convert magnitude to db
 			ArrayMagnitudeToDecibelInPlace(Out, -96.0f);
 			// compute 1 gain and use for all channels ...	const float ComputedGain = ComputeGain(DetectorOutLinkedDb);
-			ComputeGains(OutEnvelopes[0], InNumSamples);
+			ComputeGains(OutEnvelopes[0], InNumFrames);
 			NumGainChannels = 1;
 		}
 		break;
@@ -459,29 +490,38 @@ namespace Audio
 		case EDynamicsProcessorChannelLinkMode::Disabled:
 		default:
 		{
+			TRACE_CPUPROFILER_EVENT_SCOPE(FDynamicsProcessor::ProcessAudioDeinterleaveInternal_LinkModeDisabled);
 			for (int32 Channel = 0; Channel < KeyNumChannels; ++Channel)
 			{
 				// convert magnitude to db
-				ArrayMagnitudeToDecibelInPlace(TArrayView<float>(OutEnvelopes[Channel],InNumSamples), -96.0f);
+				ArrayMagnitudeToDecibelInPlace(TArrayView<float>(OutEnvelopes[Channel],InNumFrames), -96.0f);
 				// compute 1 gain and use for all channels ...	const float ComputedGain = ComputeGain(DetectorOutLinkedDb);
-				ComputeGains(OutEnvelopes[Channel], InNumSamples);
+				ComputeGains(OutEnvelopes[Channel], InNumFrames);
 			}
 			NumGainChannels = KeyNumChannels;
 		}
 		break;
 		}
 
-		const float OutAndInGain = OutputGain * InputGain;
-		for (int32 Channel = 0; Channel < NumChannels; ++Channel)
 		{
-			int32 GainChannel = Channel % NumGainChannels;
+			TRACE_CPUPROFILER_EVENT_SCOPE(FDynamicsProcessor::ProcessAudioDeinterleaveInternal_ApplyGain);
 
-			// copy input to look ahead delay write position
-			LookaheadDelay[Channel].ProcessAudioBuffer(InBuffers[Channel], InNumSamples, OutBuffers[Channel]);
-			
-			// apply compression gain & output gain & input gain to look ahead delay read position and save to output
-			ArrayMultiplyInPlace(TArrayView<const float>(OutEnvelopes[GainChannel], InNumSamples), TArrayView<float>(OutBuffers[Channel], InNumSamples));
-			ArrayMultiplyByConstant(TArrayView<const float>(OutBuffers[Channel], InNumSamples), OutAndInGain, TArrayView<float>(OutBuffers[Channel], InNumSamples));
+			const float OutAndInGain = OutputGain * InputGain;
+			for (int32 Channel = 0; Channel < NumChannels; ++Channel)
+			{
+				int32 GainChannel = Channel % NumGainChannels;
+
+				TArrayView<const float> InBufferView = MakeArrayView<const float>(InBuffers[Channel], InNumFrames);
+				TArrayView<float> OutBufferView = MakeArrayView<float>(OutBuffers[Channel], InNumFrames);
+				TArrayView<const float> EnvelopeBufferView = MakeArrayView<const float>(OutEnvelopes[GainChannel], InNumFrames);
+
+				// copy input to look ahead delay write position
+				LookaheadDelay[Channel].ProcessAudio(InBufferView, OutBufferView);
+				
+				// apply compression gain & output gain & input gain to look ahead delay read position and save to output
+				ArrayMultiplyInPlace(EnvelopeBufferView, OutBufferView);
+				ArrayMultiplyByConstantInPlace(OutBufferView, OutAndInGain);
+			}
 		}
 	}
 
@@ -657,6 +697,11 @@ namespace Audio
 		}
 	}
 
+	int32 FDynamicsProcessor::GetNumDelayFrames() const
+	{
+		return FMath::Max(0, FMath::CeilToInt(LookaheadDelayMsec * SampleRate / 1000.0f));
+	}
+
 	void FDynamicsProcessor::CalculateSlope()
 	{
 		SlopeFactor = 0.0f;
diff --git a/Engine/Source/Runtime/SignalProcessing/Private/FloatArrayMath.cpp b/Engine/Source/Runtime/SignalProcessing/Private/FloatArrayMath.cpp
--- a/Engine/Source/Runtime/SignalProcessing/Private/FloatArrayMath.cpp
+++ b/Engine/Source/Runtime/SignalProcessing/Private/FloatArrayMath.cpp
@@ -2341,16 +2341,18 @@ namespace Audio
 		ArrayInterleave(InBufferPtr, OutBuffer.GetData(), NumFrames, NumChannels);
 	}
 
-	void ArrayInterleave(const float** RESTRICT InBuffers, float* RESTRICT OutBuffer, const int32 InFrames, const int32 InChannels)
+	void ArrayInterleave(const float* const* RESTRICT InBuffers, float* RESTRICT OutBuffer, const int32 InFrames, const int32 InChannels)
 	{
 		CSV_SCOPED_TIMING_STAT(Audio_Dsp, ArrayInterleave);
 		for(int32 ChannelIdx = 0; ChannelIdx < InChannels; ChannelIdx++)
 		{
-			const float* InBuffer = InBuffers[ChannelIdx];
+			const float* InPtr = InBuffers[ChannelIdx];
+			float* OutPtr = &OutBuffer[ChannelIdx];
 			
 			for(int32 SampleIdx = 0; SampleIdx < InFrames; SampleIdx++)
 			{
-				OutBuffer[ChannelIdx + (SampleIdx * InChannels)] = InBuffer[SampleIdx];
+				*OutPtr = *InPtr++;
+				OutPtr += InChannels;
 			}
 		}
 	}
@@ -2379,17 +2381,19 @@ namespace Audio
 		ArrayDeinterleave(InBuffer.GetData(), OutBufferPtr, NumFrames, InChannels);
 	}
 
-	void ArrayDeinterleave(const float* RESTRICT InBuffer, float** RESTRICT OutBuffers, const int32 InFrames, const int32 InChannels)
+	void ArrayDeinterleave(const float* RESTRICT InBuffer, float* const* RESTRICT OutBuffers, const int32 InFrames, const int32 InChannels)
 	{
 		CSV_SCOPED_TIMING_STAT(Audio_Dsp, ArrayDeinterleave);
 
 		for(int32 ChannelIdx = 0; ChannelIdx < InChannels; ChannelIdx++)
 		{
-			float* OutBuffer = OutBuffers[ChannelIdx];
+			const float* InPtr = &InBuffer[ChannelIdx];
+			float* OutPtr = OutBuffers[ChannelIdx];
 			
 			for(int32 SampleIdx = 0; SampleIdx < InFrames; SampleIdx++)
 			{
-				OutBuffer[SampleIdx] = InBuffer[ChannelIdx + (SampleIdx * InChannels)];
+				*OutPtr++ = *InPtr;
+				InPtr += InChannels;
 			}
 		}
 	}
diff --git a/Engine/Source/Runtime/SignalProcessing/Private/IntegerDelay.cpp b/Engine/Source/Runtime/SignalProcessing/Private/IntegerDelay.cpp
--- a/Engine/Source/Runtime/SignalProcessing/Private/IntegerDelay.cpp
+++ b/Engine/Source/Runtime/SignalProcessing/Private/IntegerDelay.cpp
@@ -67,12 +67,20 @@ namespace Audio
 
 	void FIntegerDelay::ProcessAudio(const Audio::FAlignedFloatBuffer& InSamples, Audio::FAlignedFloatBuffer& OutSamples)
 	{
-		const float* InSampleData = InSamples.GetData();
-		const int32 InNum = InSamples.Num();
-
 		// Prepare output buffer
-		OutSamples.Reset(InNum);
-		OutSamples.AddUninitialized(InNum);
+		const int32 Num = InSamples.Num();
+		OutSamples.Reset(Num);
+		OutSamples.AddUninitialized(Num);
+
+		ProcessAudio(TArrayView<const float>(InSamples.GetData(), InSamples.Num()), TArrayView<float>(OutSamples.GetData(), OutSamples.Num()));
+	}
+
+	void FIntegerDelay::ProcessAudio(TArrayView<const float> InSamples, TArrayView<float> OutSamples)
+	{
+		check(InSamples.Num() == OutSamples.Num());
+
+		const int32 InNum = InSamples.Num();
+		const float* InSampleData = InSamples.GetData();
 		float* OutSampleData = OutSamples.GetData();
 
 		// Process audio one block at a time.
@@ -89,7 +97,6 @@ namespace Audio
 
 	void FIntegerDelay::ProcessAudioBlock(const float* InSamples, const int32 InNum, float* OutSamples)
 	{
-		
 		// Update delay line.	
 		DelayLine->AddSamples(InSamples, InNum);
 
diff --git a/Engine/Source/Runtime/SignalProcessing/Public/DSP/DynamicsProcessor.h b/Engine/Source/Runtime/SignalProcessing/Public/DSP/DynamicsProcessor.h
--- a/Engine/Source/Runtime/SignalProcessing/Public/DSP/DynamicsProcessor.h
+++ b/Engine/Source/Runtime/SignalProcessing/Public/DSP/DynamicsProcessor.h
@@ -4,7 +4,7 @@
 
 #include "CoreMinimal.h"
 #include "DSP/EnvelopeFollower.h"
-#include "DSP/Delay.h"
+#include "DSP/IntegerDelay.h"
 #include "Filter.h"
 
 namespace Audio
@@ -70,11 +70,10 @@ namespace Audio
 
 		SIGNALPROCESSING_API void ProcessAudioFrame(const float* InFrame, float* OutFrame, const float* InKeyFrame);
 		SIGNALPROCESSING_API void ProcessAudioFrame(const float* InFrame, float* OutFrame, const float* InKeyFrame, float* OutGain);
-		SIGNALPROCESSING_API void ProcessAudio(const float* InBuffer, const int32 InNumSamples, float* OutBuffer, const float* InKeyBuffer = nullptr);
-		SIGNALPROCESSING_API void ProcessAudio(const float* InBuffer, const int32 InNumSamples, float* OutBuffer, const float* InKeyBuffer, float* OutEnvelope);
+		SIGNALPROCESSING_API void ProcessAudio(const float* InBuffer, const int32 InNumSamples, float* OutBuffer, const float* InKeyBuffer = nullptr, float* OutEnvelope = nullptr);
 		 
 		// For single channels of audio OR non-interleaved blocks of multichannel audio.
-		SIGNALPROCESSING_API void ProcessAudio(const float* const* const InBuffers, const int32 InNumSamples, float* const* OutBuffers, const float* const* const InKeyBuffers, float* const* OutEnvelopes);
+		SIGNALPROCESSING_API void ProcessAudio(const float* const* const InBuffers, const int32 InNumFrames, float* const* OutBuffers, const float* const* const InKeyBuffers, float* const* OutEnvelopes);
 
 	protected:
 		SIGNALPROCESSING_API float ComputeGain(const float InEnvFollowerDb);
@@ -100,7 +99,7 @@ namespace Audio
 		EPeakMode::Type EnvelopeFollowerPeakMode;
 
 		// Lookahead delay lines
-		TArray<FDelay> LookaheadDelay;
+		TArray<FIntegerDelay> LookaheadDelay;
 
 		// Envelope followers
 		TArray<FInlineEnvelopeFollower> EnvFollower;
@@ -112,7 +111,7 @@ namespace Audio
 		TArray<float> Gain;
 
 		// How far ahead to look in the audio
-		float LookaheedDelayMsec;
+		float LookaheadDelayMsec;
 
 		// The period of which the compressor decreases gain to the level determined by the compression ratio
 		float AttackTimeMsec;
@@ -161,6 +160,15 @@ namespace Audio
 		static constexpr float MaxLookaheadMsec = 100.0f;
 
 	private:
+
+		// OutBuffers is also used a temporary buffer for processing the Key. 
+		// Thus, OutBuffers must supply enough valid memory to support it's use as a 
+		// temporary buffer and needs to have Max(NumInputChannels, NumKeyChannels) 
+		// buffers available.  
+		void ProcessAudioDeinterleaveInternal(const float* const* const InBuffers, const int32 InNumFrames, float* const* OutBuffers, const float* const* const InKeyBuffers, float* const* OutEnvelopes);
+
+		int32 GetNumDelayFrames() const;
+
 		void CalculateSlope();
 
 		void CalculateKnee();
diff --git a/Engine/Source/Runtime/SignalProcessing/Public/DSP/FloatArrayMath.h b/Engine/Source/Runtime/SignalProcessing/Public/DSP/FloatArrayMath.h
--- a/Engine/Source/Runtime/SignalProcessing/Public/DSP/FloatArrayMath.h
+++ b/Engine/Source/Runtime/SignalProcessing/Public/DSP/FloatArrayMath.h
@@ -239,13 +239,13 @@ namespace Audio
 	SIGNALPROCESSING_API void ArrayInterleave(const TArray<FAlignedFloatBuffer>& InBuffers, FAlignedFloatBuffer& OutBuffer);
 
 	/** Interleaves samples from an array of input buffers */
-	SIGNALPROCESSING_API void ArrayInterleave(const float** RESTRICT InBuffers, float* RESTRICT OutBuffer, const int32 InFrames, const int32 InChannels);
+	SIGNALPROCESSING_API void ArrayInterleave(const float* const* RESTRICT InBuffers, float* RESTRICT OutBuffer, const int32 InFrames, const int32 InChannels);
 
 	/** Interleaves samples from an array of input buffers */
 	SIGNALPROCESSING_API void ArrayDeinterleave(const FAlignedFloatBuffer& InBuffer, TArray<FAlignedFloatBuffer>& OutBuffers, const int32 InChannels);
 
 	/** Interleaves samples from an array of input buffers */
-	SIGNALPROCESSING_API void ArrayDeinterleave(const float* RESTRICT InBuffer, float** RESTRICT OutBuffers, const int32 InFrames, const int32 InChannels);
+	SIGNALPROCESSING_API void ArrayDeinterleave(const float* RESTRICT InBuffer, float* const* RESTRICT OutBuffers, const int32 InFrames, const int32 InChannels);
 
 	/** Interpolates a Mono audio buffer. */
 	SIGNALPROCESSING_API void ArrayInterpolate(const float* RESTRICT InBuffer, float* RESTRICT OutBuffer, const int32 NumInSamples, const int32 NumOutSamples);
diff --git a/Engine/Source/Runtime/SignalProcessing/Public/DSP/IntegerDelay.h b/Engine/Source/Runtime/SignalProcessing/Public/DSP/IntegerDelay.h
--- a/Engine/Source/Runtime/SignalProcessing/Public/DSP/IntegerDelay.h
+++ b/Engine/Source/Runtime/SignalProcessing/Public/DSP/IntegerDelay.h
@@ -32,7 +32,23 @@ namespace Audio
 		SIGNALPROCESSING_API int32 GetNumDelaySamples() const;
 
 		// Process InSamples, placing delayed versions in OutSamples.
+		FORCEINLINE float ProcessAudioSample(float InSample)
+		{
+			// Update delay line.	
+			DelayLine->AddSamples(&InSample, 1);
+
+			// Copy delayed version to output
+			const float DelayedSample = *(DelayLine->InspectSamples(1 + NumBufferOffsetSamples, NumDelayLineOffsetSamples));
+
+			// Remove unneeded delay line.
+			DelayLine->RemoveSamples(1);
+
+			return DelayedSample;
+		}
+
 		SIGNALPROCESSING_API void ProcessAudio(const Audio::FAlignedFloatBuffer& InSamples, Audio::FAlignedFloatBuffer& OutSamples);
+		// Process InSamples, placing delayed versions in OutSamples.
+		SIGNALPROCESSING_API void ProcessAudio(TArrayView<const float> InSamples, TArrayView<float> OutSamples);
 
 		// Retrieve a copy of the internal delay line.
 		// InNum must be less than or equal to InMaxNumDelaySamples.
