Audio::ArrayMixIn loop unrolling for SSE optimization

diff --git a/Engine/Source/Runtime/AudioMixer/Private/AudioMixerSubmix.cpp b/Engine/Source/Runtime/AudioMixer/Private/AudioMixerSubmix.cpp
--- a/Engine/Source/Runtime/AudioMixer/Private/AudioMixerSubmix.cpp
+++ b/Engine/Source/Runtime/AudioMixer/Private/AudioMixerSubmix.cpp
@@ -1565,10 +1565,14 @@ namespace Audio
 			bool bIsNowSilent = true;
 			int i = 0;
 #if PLATFORM_ENABLE_VECTORINTRINSICS
-			for (; i < OutAudioBuffer.Num(); i += 4)
-			{
-				VectorRegister4Float Samples = VectorLoad(&OutAudioBuffer[i]);
-				if (VectorAnyGreaterThan(Samples, GlobalVectorConstants::SmallNumber))
+			int SimdNum = OutAudioBuffer.Num() & 0xFFFFFFF0;
+			for (; i < SimdNum; i += 16)
+			{
+				VectorRegister4x4Float Samples = VectorLoad16(&OutAudioBuffer[i]);
+				if (   VectorAnyGreaterThan(VectorAbs(Samples.val[0]), GlobalVectorConstants::SmallNumber)
+					|| VectorAnyGreaterThan(VectorAbs(Samples.val[1]), GlobalVectorConstants::SmallNumber)
+					|| VectorAnyGreaterThan(VectorAbs(Samples.val[2]), GlobalVectorConstants::SmallNumber)
+					|| VectorAnyGreaterThan(VectorAbs(Samples.val[3]), GlobalVectorConstants::SmallNumber))
 				{
 					bIsNowSilent = false;
 					i = INT_MAX;
@@ -1581,7 +1585,7 @@ namespace Audio
 			for (; i < OutAudioBuffer.Num(); ++i)
 			{
 				// As soon as we hit a non-silent sample, we're not silent
-				if (OutAudioBuffer[i] > SMALL_NUMBER)
+				if (FMath::Abs(OutAudioBuffer[i]) > SMALL_NUMBER)
 				{
 					bIsNowSilent = false;
 					break;
diff --git a/Engine/Source/Runtime/Core/Public/Math/UnrealMathFPU.h b/Engine/Source/Runtime/Core/Public/Math/UnrealMathFPU.h
--- a/Engine/Source/Runtime/Core/Public/Math/UnrealMathFPU.h
+++ b/Engine/Source/Runtime/Core/Public/Math/UnrealMathFPU.h
@@ -93,6 +93,11 @@ typedef VectorRegister4Double VectorRegister4;
 typedef VectorRegister4 VectorRegister;
 typedef VectorRegister4Int VectorRegisterInt;
 
+typedef struct
+{
+	VectorRegister4Float val[4];
+} VectorRegister4x4Float;
+
 
 // Forward declarations
 VectorRegister4Float VectorLoadAligned(const float* Ptr);
@@ -386,6 +391,22 @@ FORCEINLINE VectorRegister4Double VectorLoad(const double* Ptr)
 	return MakeVectorRegisterDouble(Ptr[0], Ptr[1], Ptr[2], Ptr[3]);
 }
 
+/**
+ * Loads 16 floats from unaligned memory into 4 vector registers.
+ *
+ * @param Ptr	Unaligned memory pointer to the 4 floats
+ * @return		VectorRegister4x4Float containing 16 floats
+ */
+FORCEINLINE VectorRegister4x4Float VectorLoad16(const float* Ptr)
+{
+	VectorRegister4x4Float Result;
+	Result.val[0] = VectorLoad(Ptr);
+	Result.val[1] = VectorLoad(Ptr + 4);
+	Result.val[2] = VectorLoad(Ptr + 8);
+	Result.val[3] = VectorLoad(Ptr + 12);
+	return Result;
+}
+
 
 /**
  * Loads 3 FLOATs from unaligned memory and sets W=0.
@@ -542,6 +563,17 @@ FORCEINLINE void VectorStore(const VectorRegister4Double& Vec, double* Dst)
 	FMemory::Memcpy(Dst, &(Vec), 4 * sizeof(double));
 }
 
+/**
+ * Stores 4 vectors to memory (aligned or unaligned).
+ *
+ * @param Vec	Vector to store
+ * @param Ptr	Memory pointer
+ */
+FORCEINLINE void VectorStore16(const VectorRegister4x4Float& Vec, float* Dst)
+{
+	FMemory::Memcpy(Dst, &(Vec), 16 * sizeof(float));
+}
+
 
 /**
  * Stores the XYZ components of a vector to unaligned memory.
diff --git a/Engine/Source/Runtime/Core/Public/Math/UnrealMathNeon.h b/Engine/Source/Runtime/Core/Public/Math/UnrealMathNeon.h
--- a/Engine/Source/Runtime/Core/Public/Math/UnrealMathNeon.h
+++ b/Engine/Source/Runtime/Core/Public/Math/UnrealMathNeon.h
@@ -109,6 +109,7 @@ typedef float32x4_t GCC_ALIGN(16) VectorRegister4Float;
 typedef float64x2_t GCC_ALIGN(16) VectorRegister2Double;
 typedef int32x4_t  GCC_ALIGN(16) VectorRegister4Int;
 typedef int64x2_t GCC_ALIGN(16) VectorRegister2Int64;
+typedef float32x4x4_t GCC_ALIGN(16) VectorRegister4x4Float;
 
 FORCEINLINE constexpr VectorRegister4Int MakeVectorRegisterIntConstant(int32 X, int32 Y, int32 Z, int32 W)
 {
@@ -472,6 +473,17 @@ FORCEINLINE VectorRegister4Double VectorLoad(const double* Ptr)
 	return Result;
 }
 
+/**
+ * Loads 16 floats from unaligned memory into 4 vector registers.
+ *
+ * @param Ptr	Unaligned memory pointer to the 4 floats
+ * @return		VectorRegister4x4Float containing 16 floats
+ */
+FORCEINLINE VectorRegister4x4Float VectorLoad16(const float* Ptr)
+{
+	return vld1q_f32_x4(Ptr);
+}
+
 /**
  * Loads 2 floats from unaligned memory into X and Y and duplicates them in Z and W.
  *
@@ -681,6 +693,17 @@ FORCEINLINE void VectorStore(const VectorRegister4Double& Vec, double* Ptr)
 	vst1q_f64_x2(Ptr, *(float64x2x2_t*)&Vec);
 }
 
+/**
+ * Stores 4 vectors to memory (aligned or unaligned).
+ *
+ * @param Vec	Vector to store
+ * @param Ptr	Memory pointer
+ */
+FORCEINLINE void VectorStore16(const VectorRegister4x4Float& Vec, float* Ptr)
+{
+	vst1q_f32_x4(Ptr, Vec);
+}
+
 /**
  * Stores the XYZ components of a vector to unaligned memory.
  *
diff --git a/Engine/Source/Runtime/Core/Public/Math/UnrealMathSSE.h b/Engine/Source/Runtime/Core/Public/Math/UnrealMathSSE.h
--- a/Engine/Source/Runtime/Core/Public/Math/UnrealMathSSE.h
+++ b/Engine/Source/Runtime/Core/Public/Math/UnrealMathSSE.h
@@ -87,6 +87,12 @@ typedef __m128i VectorRegister2Int64;
 // 2 doubles
 typedef __m128d	VectorRegister2Double;
 
+typedef struct
+{
+	//TODO: alias for AVX2!
+	VectorRegister4Float val[4];
+} VectorRegister4x4Float;
+
 
 namespace SSE
 {
@@ -557,6 +563,22 @@ FORCEINLINE VectorRegister4Double VectorLoad(const double* Ptr)
 	return Result;
 }
 
+/**
+ * Loads 16 floats from unaligned memory into 4 vector registers.
+ *
+ * @param Ptr	Unaligned memory pointer to the 4 floats
+ * @return		VectorRegister4x4Float containing 16 floats
+ */
+FORCEINLINE VectorRegister4x4Float VectorLoad16(const float* Ptr)
+{
+	VectorRegister4x4Float Result;
+	Result.val[0] = VectorLoad(Ptr);
+	Result.val[1] = VectorLoad(Ptr + 4);
+	Result.val[2] = VectorLoad(Ptr + 8);
+	Result.val[3] = VectorLoad(Ptr + 12);
+	return Result;
+}
+
 /**
  * Loads 3 FLOATs from unaligned memory and sets W=0.
  *
@@ -752,6 +774,20 @@ FORCEINLINE void VectorStore(const VectorRegister4Double& Vec, double* Dst)
 #endif
 }
 
+/**
+ * Stores 4 vectors to memory (aligned or unaligned).
+ *
+ * @param Vec	Vector to store
+ * @param Ptr	Memory pointer
+ */
+FORCEINLINE void VectorStore16(const VectorRegister4x4Float& Vec, float* Ptr)
+{
+	VectorStore(Vec.val[0], Ptr);
+	VectorStore(Vec.val[1], Ptr + 4);
+	VectorStore(Vec.val[2], Ptr + 8);
+	VectorStore(Vec.val[3], Ptr + 12);
+}
+
 /**
  * Stores a vector to aligned memory.
  *
diff --git a/Engine/Source/Runtime/SignalProcessing/Private/FloatArrayMath.cpp b/Engine/Source/Runtime/SignalProcessing/Private/FloatArrayMath.cpp
--- a/Engine/Source/Runtime/SignalProcessing/Private/FloatArrayMath.cpp
+++ b/Engine/Source/Runtime/SignalProcessing/Private/FloatArrayMath.cpp
@@ -35,6 +35,12 @@ namespace Audio
 		const float Loge10 = FMath::Loge(10.f);
 		const int32 SimdMask = 0xFFFFFFFC;
 		const int32 NotSimdMask = 0x00000003;
+
+		const int32 Simd8Mask = 0xFFFFFFF8;
+		const int32 NotSimd8Mask = 0x00000007;
+
+		const int32 Simd16Mask = 0xFFFFFFF0;
+		const int32 NotSimd16Mask = 0x0000000F;
 	}
 
 	void ArraySum(TArrayView<const float> InValues, float& OutSum)
@@ -2032,27 +2038,18 @@ namespace Audio
 		{
 			VectorRegister4Float GainVector = VectorLoadFloat1(&Gain);
 			int32 i = 0;
-#if PLATFORM_ENABLE_VECTORINTRINSICS_NEON
-			// this approach is ~10% faster than the other one on NEON
-			for (; i < Num; i += 16)
+			const int32 SimdNum = Num & MathIntrinsics::Simd16Mask;
+			for (; i < SimdNum; i += 16)
 			{
-				float32x4x4_t Input = vld1q_f32_x4(&InData[i]);
-				float32x4x4_t Output = vld1q_f32_x4(&InOutData[i]);
+				// manually unrolling the loop produces a bit faster code
+				VectorRegister4x4Float Input = VectorLoad16(&InData[i]);
+				VectorRegister4x4Float Output = VectorLoad16(&InOutData[i]);
 				Output.val[0] = VectorMultiplyAdd(Input.val[0], GainVector, Output.val[0]);
 				Output.val[1] = VectorMultiplyAdd(Input.val[1], GainVector, Output.val[1]);
 				Output.val[2] = VectorMultiplyAdd(Input.val[2], GainVector, Output.val[2]);
 				Output.val[3] = VectorMultiplyAdd(Input.val[3], GainVector, Output.val[3]);
-				vst1q_f32_x4(&InOutData[i], Output);
-			}
-#else
-			for (; i < Num; i += AUDIO_NUM_FLOATS_PER_VECTOR_REGISTER)
-			{
-				VectorRegister4Float Input = VectorLoad(&InData[i]);
-				VectorRegister4Float Output = VectorLoad(&InOutData[i]);
-				Output = VectorMultiplyAdd(Input, GainVector, Output);
-				VectorStore(Output, &InOutData[i]);
+				VectorStore16(Output, &InOutData[i]);
 			}
-#endif //~PLATFORM_ENABLE_VECTORINTRINSICS_NEON
 
 			for (; i < Num; ++i)
 			{
@@ -2080,26 +2077,18 @@ namespace Audio
 		else
 		{
 			int32 i = 0;
-#if PLATFORM_ENABLE_VECTORINTRINSICS_NEON
-			for (; i < Num; i += 16)
+			const int32 SimdNum = Num & MathIntrinsics::Simd16Mask;
+			for (; i < SimdNum; i += 16)
 			{
-				float32x4x4_t Input = vld1q_f32_x4(&InData[i]);
-				float32x4x4_t Output = vld1q_f32_x4(&InOutData[i]);
+				// manually unrolling the loop produces a bit faster code
+				VectorRegister4x4Float Input = VectorLoad16(&InData[i]);
+				VectorRegister4x4Float Output = VectorLoad16(&InOutData[i]);
 				Output.val[0] = VectorAdd(Input.val[0], Output.val[0]);
 				Output.val[1] = VectorAdd(Input.val[1], Output.val[1]);
 				Output.val[2] = VectorAdd(Input.val[2], Output.val[2]);
 				Output.val[3] = VectorAdd(Input.val[3], Output.val[3]);
-				vst1q_f32_x4(&InOutData[i], Output);
-			}
-#else
-			for (; i < Num; i += AUDIO_NUM_FLOATS_PER_VECTOR_REGISTER)
-			{
-				VectorRegister4Float Output = VectorLoad(&InOutData[i]);
-				VectorRegister4Float Input = VectorLoad(&InData[i]);
-				Output = VectorAdd(Input, Output);
-				VectorStore(Output, &InOutData[i]);
+				VectorStore16(Output, &InOutData[i]);
 			}
-#endif
 
 			for (; i < Num; ++i)
 			{
@@ -2235,19 +2224,20 @@ namespace Audio
 		const VectorRegister4Float Multiplier = VectorSetFloat1(ConversionValue);
 		int32 i = 0;
 #if PLATFORM_ENABLE_VECTORINTRINSICS_NEON
-		for (; i < Num; i += 8)
+		const int32 SimdNum = Num & MathIntrinsics::Simd8Mask;
+		for (; i < SimdNum; i += 8)
 		{
-			const VectorRegister4Float InVector1 = VectorLoad(&InputPtr[i]);
-			const VectorRegister4Float InVector2 = VectorLoad(&InputPtr[i + 4]);
-			const VectorRegister4Float ScaledVector1 = VectorMultiply(InVector1, Multiplier);
-			const VectorRegister4Float ScaledVector2 = VectorMultiply(InVector2, Multiplier);
+			const float32x4x2_t InVector = vld1q_f32_x2(&InputPtr[i]);
+			const VectorRegister4Float ScaledVector1 = VectorMultiply(InVector.val[0], Multiplier);
+			const VectorRegister4Float ScaledVector2 = VectorMultiply(InVector.val[1], Multiplier);
 			const VectorRegister4Int IntVector1 = VectorFloatToInt(ScaledVector1);
 			const VectorRegister4Int IntVector2 = VectorFloatToInt(ScaledVector2);
 			const int16x8_t Result = vmovn_high_s32(vmovn_u32(IntVector1), IntVector2);
 			vst1q_s16(&OutPtr[i], Result);
 		}
 #else
-		for (; i < Num; i += AUDIO_NUM_FLOATS_PER_VECTOR_REGISTER)
+		const int32 SimdNum = Num & MathIntrinsics::SimdMask;
+		for (; i < SimdNum; i += AUDIO_NUM_FLOATS_PER_VECTOR_REGISTER)
 		{
 			const VectorRegister4Float InVector = VectorLoad(&InputPtr[i]);
 			const VectorRegister4Float ScaledVector = VectorMultiply(InVector, Multiplier);
@@ -2255,10 +2245,10 @@ namespace Audio
 
 			const AlignedFloat4 ScaledFloatArray(ScaledVector);
 
-			OutPtr[i] =		(int16)ScaledFloatArray[0];
-			OutPtr[i + 1] =	(int16)ScaledFloatArray[1];
-			OutPtr[i + 2] =	(int16)ScaledFloatArray[2];
-			OutPtr[i + 3] =	(int16)ScaledFloatArray[3];
+			OutPtr[i + 0] = (int16)ScaledFloatArray[0];
+			OutPtr[i + 1] = (int16)ScaledFloatArray[1];
+			OutPtr[i + 2] = (int16)ScaledFloatArray[2];
+			OutPtr[i + 3] = (int16)ScaledFloatArray[3];
 		}
 #endif //~PLATFORM_ENABLE_VECTORINTRINSICS_NEON
 
@@ -2284,19 +2274,21 @@ namespace Audio
 
 		int32 i = 0;
 #if PLATFORM_ENABLE_VECTORINTRINSICS_NEON
-		for (; i < Num; i += 8)
+		const int32 SimdNum = Num & MathIntrinsics::Simd8Mask;
+		for (; i < SimdNum; i += 8)
 		{
 			int16x8_t Data = vld1q_s16(&InputPtr[i]);
 			int32x4_t VecA = vmovl_s16(vget_low_s16(Data));
 			int32x4_t VecB = vmovl_high_s16(Data);
-			VectorRegister4Float FloatVecA = VectorMultiply(vcvtq_f32_s32(VecA), Multiplier);
-			VectorRegister4Float FloatVecB = VectorMultiply(vcvtq_f32_s32(VecB), Multiplier);
-			VectorStore(FloatVecA, &OutPtr[i]);
-			VectorStore(FloatVecB, &OutPtr[i + 4]);
+			float32x4x2_t FloatVec;
+			FloatVec.val[0] = VectorMultiply(vcvtq_f32_s32(VecA), Multiplier);
+			FloatVec.val[1] = VectorMultiply(vcvtq_f32_s32(VecB), Multiplier);
+			vst1q_f32_x2(&OutPtr[i], FloatVec);
 		}
 #else
 		AlignedFloat4 FloatArray(GlobalVectorConstants::FloatZero);
-		for (; i < Num; i += AUDIO_NUM_FLOATS_PER_VECTOR_REGISTER)
+		const int32 SimdNum = Num & MathIntrinsics::SimdMask;
+		for (; i < SimdNum; i += AUDIO_NUM_FLOATS_PER_VECTOR_REGISTER)
 		{
 			FloatArray[0] = (float)InputPtr[i];
 			FloatArray[1] = (float)InputPtr[i + 1];
