Audio Buffer Diagnostics

diff --git a/Engine/Source/Runtime/AudioMixer/Private/AudioMixerSubmix.cpp b/Engine/Source/Runtime/AudioMixer/Private/AudioMixerSubmix.cpp
--- a/Engine/Source/Runtime/AudioMixer/Private/AudioMixerSubmix.cpp
+++ b/Engine/Source/Runtime/AudioMixer/Private/AudioMixerSubmix.cpp
@@ -16,6 +16,8 @@
 #include "Misc/ScopeTryLock.h"
 #include "ProfilingDebugging/CsvProfiler.h"
 #include "AudioLinkLog.h"
+#include "DSP/BufferDiagnostics.h"
+#include "Algo/Accumulate.h"
 
 // Link to "Audio" profiling category
 CSV_DECLARE_CATEGORY_MODULE_EXTERN(AUDIOMIXERCORE_API, Audio);
@@ -1296,6 +1298,9 @@ namespace Audio
 					if (ChildSubmix->IsRenderingAudio())
 					{
 						ChildSubmix->ProcessAudio(InputBuffer);
+
+						// Check the buffer after processing to catch any bad values.
+						AUDIO_CHECK_BUFFER_NAMED_MSG(InputBuffer, TEXT("Submix Chidren"), TEXT("Submix: %s"), *ChildSubmix->GetName());
 					}
 				}
 				else
@@ -1322,6 +1327,9 @@ namespace Audio
 				const EMixerSourceSubmixSendStage SubmixSendStage = MixerSourceVoiceIter.Value.SubmixSendStage;
 
 				MixerSourceVoice->MixOutputBuffers(NumChannels, SendLevel, SubmixSendStage, InputBuffer);
+				
+				// Check the buffer after each voice mix to catch any bad values.
+				AUDIO_CHECK_BUFFER_NAMED_MSG(InputBuffer, TEXT("Submix SourceMix"), TEXT("Submix: %s"), *GetName());
 			}
 		}
 
@@ -1679,6 +1687,9 @@ namespace Audio
 
 			if (SubmixEffect->ProcessAudio(InputData, OutputData))
 			{
+				AUDIO_CHECK_BUFFER_NAMED_MSG(*OutputData.AudioBuffer,TEXT("Submix Effects"), TEXT("FxPreset=%s, Submix=%s"), 
+					*GetNameSafe(SubmixEffect->GetPreset()), *GetName());
+				
 				// Mix in the dry signal directly
 				const float DryLevel = SubmixEffect->GetDryLevel();
 				if (DryLevel > 0.0f)
diff --git a/Engine/Source/Runtime/SignalProcessing/Private/BufferDiagnostics.cpp b/Engine/Source/Runtime/SignalProcessing/Private/BufferDiagnostics.cpp
new file mode 100644
--- /dev/null
+++ b/Engine/Source/Runtime/SignalProcessing/Private/BufferDiagnostics.cpp
@@ -0,0 +1,425 @@
+// Copyright Epic Games, Inc. All Rights Reserved.
+#include "DSP/BufferDiagnostics.h"
+
+#include "Algo/Accumulate.h"
+#include "Algo/Find.h"
+#include "Algo/IndexOf.h"
+#include "CoreGlobals.h"
+#include "HAL/IConsoleManager.h"
+#include "Misc/VarArgs.h"
+
+namespace Audio
+{
+	// denormalized
+	int32 FindDenormalized(TArrayView<const float> InBuffer)
+	{
+		return Algo::IndexOfByPredicate(InBuffer, [](const float InSample) { return IsDenormalized(InSample); });
+	}
+
+	bool ContainsDenormalized(TArrayView<const float> InBuffer)
+	{
+		return FindDenormalized(InBuffer) != INDEX_NONE;
+	}
+
+	// +inf/-inf
+	int32 FindInfinity(TArrayView<const float> InBuffer)
+	{
+		return Algo::IndexOfByPredicate(InBuffer, [](const float InSample) { return IsInfinity(InSample); });
+	}
+
+	bool ContainsInfinity(TArrayView<const float> InBuffer)
+	{
+		return FindInfinity(InBuffer) != INDEX_NONE;
+	}
+
+	// NaN
+	int32 FindNan(TArrayView<const float> InBuffer)
+	{
+		return Algo::IndexOfByPredicate(InBuffer, [](const float InSample) { return FMath::IsNaN(InSample); });
+	}
+
+	bool ContainsNan(TArrayView<const float> InBuffer)
+	{
+		return FindNan(InBuffer) != INDEX_NONE;
+	}
+
+	// Because our find checks are composite, once failed, ask for the flags that failed.
+	static ECheckBufferFlags ClassifyFailure(TArrayView<const float> InBuffer, const int32 InIndexOfFailure)
+	{
+		if (InIndexOfFailure == INDEX_NONE || InBuffer.Num() == 0)
+		{
+			return ECheckBufferFlags::None;
+		}
+
+		ECheckBufferFlags FailFlags = ECheckBufferFlags::None;
+		const float BadValue = InBuffer[InIndexOfFailure];
+		if (IsDenormalized(BadValue))
+		{
+			FailFlags |= ECheckBufferFlags::Denormalized;
+		}
+		if (FMath::IsNaN(BadValue))
+		{
+			FailFlags |= ECheckBufferFlags::Nan;
+		}
+		if (IsInfinity(BadValue))
+		{
+			FailFlags |= ECheckBufferFlags::Infinity;
+		}
+		return FailFlags;
+	}
+
+	bool CheckBuffer(TArrayView<const float> InBuffer, const ECheckBufferFlags InFlags, ECheckBufferFlags& OutFailedFlags)
+	{
+		// Reset flags to start.
+		OutFailedFlags = ECheckBufferFlags::None;
+
+		// If all flags are on do combined check for speed.
+		if (InFlags == ECheckBufferFlags::All)
+		{
+			const int32 IndexOf = Algo::IndexOfByPredicate(InBuffer, [](const float InSample) -> bool
+			{
+				return IsDenormalized(InSample) || !FMath::IsFinite(InSample);
+			});
+			OutFailedFlags = ClassifyFailure(InBuffer, IndexOf);
+			return IndexOf == INDEX_NONE;
+		}
+
+		// Another common one is Inf/Nan together. !InFinite is equiv.
+		if (InFlags == (ECheckBufferFlags::Infinity | ECheckBufferFlags::Nan))
+		{
+			const int32 IndexOf = Algo::IndexOfByPredicate(InBuffer, [](const float InSample) -> bool
+			{
+				return !FMath::IsFinite(InSample);
+			});
+			OutFailedFlags = ClassifyFailure(InBuffer, IndexOf);
+			return IndexOf == INDEX_NONE;
+		}
+
+		// Otherwise do piece-meal testing.
+		if ((InFlags & ECheckBufferFlags::Nan) != ECheckBufferFlags::None)
+		{
+			if (ContainsNan(InBuffer))
+			{
+				OutFailedFlags |= ECheckBufferFlags::Nan;
+				return false;
+			}
+		}
+		if ((InFlags & ECheckBufferFlags::Denormalized) != ECheckBufferFlags::None)
+		{
+			if (ContainsDenormalized(InBuffer))
+			{
+				OutFailedFlags |= ECheckBufferFlags::Denormalized;
+				return false;
+			}
+		}
+		if ((InFlags & ECheckBufferFlags::Infinity) != ECheckBufferFlags::None)
+		{
+			if (ContainsInfinity(InBuffer))
+			{
+				OutFailedFlags |= ECheckBufferFlags::Infinity;
+				return false;
+			}
+		}
+
+		// Pass.
+		return true;
+	}
+
+	// Static linkage.
+	FCriticalSection FCheckedBufferState::ListCs;
+	FCheckedBufferState* FCheckedBufferState::Head = nullptr;
+	
+	FCheckedBufferState::FCheckedBufferState(const int32 InLine, const TCHAR* InFile, const TCHAR* InName, const ECheckBufferFlags InCheckFlags, const EBufferCheckBehavior InCheckBehavior)
+		: Name(InName)
+		, Line(InLine)
+		, File(InFile)
+		, CheckFlags(InCheckFlags)
+		, Behavior(InCheckBehavior)
+	{
+		FScopeLock Lock(&ListCs);
+		Next = Head;
+		Head = this;	
+	}
+
+	void FCheckedBufferState::ForEach(TFunctionRef<void(FCheckedBufferState&)> InCmd)
+	{
+		FScopeLock Lock(&ListCs);
+		for (FCheckedBufferState* i=Head; i; i=i->Next)
+		{
+			InCmd(*i);
+		}
+	}
+
+
+	bool FCheckedBufferState::DoCheck(TArrayView<const float> InBuffer)
+	{
+		// Only do the check if we have some flags to test against
+		// and only do the check if we've not had a prior failure.
+		if (CheckFlags != ECheckBufferFlags::None && FailFlags == ECheckBufferFlags::None)
+		{
+			if (!CheckBuffer(InBuffer, CheckFlags, FailFlags))
+			{
+				return false;
+			}
+		}
+
+		// Pass.
+		return true;
+	}
+	
+	void FCheckedBufferState::FailedBufferCheckImpl(const TCHAR* InFormat, ...) const
+	{
+		TCHAR Msg[512];
+		GET_TYPED_VARARGS(TCHAR, Msg, UE_ARRAY_COUNT(Msg), UE_ARRAY_COUNT(Msg)-1, InFormat, InFormat);
+		
+		switch (Behavior)
+		{
+			case EBufferCheckBehavior::Ensure:
+			{
+				ensureMsgf(false, TEXT("Failed Buffer Check: '%s', Failure='%s', %s:%d, CheckFlags='%s', %s"),Name,
+					*ToDelimitedString(FailFlags), File, Line, *ToDelimitedString(CheckFlags), Msg);
+				break;
+			}
+			case EBufferCheckBehavior::Log:
+			{
+				UE_LOG(LogTemp, Warning, TEXT("Failed Buffer Check: '%s', Failure='%s', %s:%d, CheckFlags='%s', %s"), Name,
+					*ToDelimitedString(FailFlags), File, Line, *ToDelimitedString(CheckFlags), Msg);
+				break;
+			}
+			case EBufferCheckBehavior::Break:	
+			{
+				UE_LOG(LogTemp, Warning, TEXT("Failed Buffer Check: '%s', Failure='%s', %s:%d, CheckFlags='%s', %s"), Name,
+					*ToDelimitedString(FailFlags), File, Line, *ToDelimitedString(CheckFlags), Msg);
+				UE_DEBUG_BREAK_AND_PROMPT_FOR_REMOTE();
+				break;
+			}
+			case EBufferCheckBehavior::Nothing:
+			default:
+			{
+				break;
+			}
+		}
+	}
+}
+
+namespace BufferDiagnosticsPrivate
+{
+	static void HandleDiagnosticsCmd(const TArray<FString>& InArgs, FOutputDevice& InDevice);
+
+	// Handle the enabling diagnostics.
+	FAutoConsoleCommandWithArgsAndOutputDevice CmdBufferDiagnostics(
+		TEXT("au.debug.bufferdiagnostics"),
+		TEXT("Enables per buffer sample diagnostics (Nans/denorms/Infs)"),
+		FConsoleCommandWithArgsAndOutputDeviceDelegate::CreateStatic(HandleDiagnosticsCmd));
+	
+	using namespace Audio;
+	
+	// Simple compile time function to remove the fully qualified name.
+	// Example "ECheckBufferFlags::Infinity" -> "Infinity".
+	static constexpr const TCHAR* ToShortEnumName(const TCHAR* InFullyQualifiedName) 
+	{
+		const TCHAR* Iter = InFullyQualifiedName;
+		for(; *Iter; ++Iter)
+		{
+			if (*Iter == L':') 
+			{
+				return Iter + 2;
+			}
+		}
+		return InFullyQualifiedName;
+	}
+#define FOREACH_ECHECKBUFFERFLAGS(OP)\
+	OP(ECheckBufferFlags::None)\
+	OP(ECheckBufferFlags::Infinity)\
+	OP(ECheckBufferFlags::Nan)\
+	OP(ECheckBufferFlags::Denormalized)\
+	OP(ECheckBufferFlags::All)
+
+#define FOREACH_ECHECKBUFFERBEHAVIOR(OP)\
+	OP(EBufferCheckBehavior::Nothing)\
+	OP(EBufferCheckBehavior::Ensure)\
+	OP(EBufferCheckBehavior::Log)\
+	OP(EBufferCheckBehavior::Break)
+	
+// Switch on the enum and return a short string version.
+#ifndef TEST_CASE_ENUM_RETURN_SHORTNAME 
+	#define TEST_CASE_ENUM_RETURN_SHORTNAME(X) case X: { return ToShortEnumName(TEXT(#X)); }
+#endif //TEST_CASE_ENUM_RETURN_SHORTNAME
+	
+	static const TCHAR* LexToString(const ECheckBufferFlags InEnum)
+	{
+		switch (InEnum)
+		{
+			FOREACH_ECHECKBUFFERFLAGS(TEST_CASE_ENUM_RETURN_SHORTNAME)
+			default: break;
+		}
+		return TEXT("Unknown ECheckBufferFlags");
+	}
+	static const TCHAR* LexToString(const EBufferCheckBehavior InEnum)
+	{
+		switch (InEnum)
+		{
+			FOREACH_ECHECKBUFFERBEHAVIOR(TEST_CASE_ENUM_RETURN_SHORTNAME)
+			default: break;
+		}
+		return TEXT("Unknown EBufferCheckBehavior");
+	}
+	template<typename EnumType> static bool IsEnumEntry(const TCHAR* InString, const EnumType InEnum)
+	{
+		return FCString::Stricmp(InString, LexToString(InEnum)) == 0;
+	}
+
+// Test string against an enum entry, return true on a match and set output property to the match.
+#ifndef TEST_STRING_AGAINST_ENUM_AND_RETURN
+	#define TEST_STRING_AGAINST_ENUM_AND_RETURN(X) if (IsEnumEntry(InStr, X)) { OutEnum=X; return true; }
+#endif //TEST_STRING_AGAINST_ENUM_AND_RETURN
+	
+	static bool TryLexFromString(const TCHAR* InStr, ECheckBufferFlags& OutEnum)
+	{
+		FOREACH_ECHECKBUFFERFLAGS(TEST_STRING_AGAINST_ENUM_AND_RETURN);
+		return false; // Fail.
+	}
+	static bool TryLexFromString(const TCHAR* InStr, EBufferCheckBehavior& OutEnum)
+	{
+		FOREACH_ECHECKBUFFERBEHAVIOR(TEST_STRING_AGAINST_ENUM_AND_RETURN)
+		return false; // Fail.
+	}
+	static void BuildBitfieldString(const ECheckBufferFlags InBitfield, const ECheckBufferFlags InFlag, FString& OutString)
+	{
+		if ((InBitfield & InFlag) == InFlag && static_cast<uint32>(InFlag) != 0)
+		{
+			OutString += LexToString(InFlag);
+			OutString += TEXT("|");
+		}	
+	}
+
+// Test enum against each flag of bitfield and build a string
+#ifndef TEST_AND_BUILD_BITFIELD_STRING 
+	#define TEST_AND_BUILD_BITFIELD_STRING(X) BuildBitfieldString(InEnum, X, String);
+#endif //TEST_AND_BUILD_BITFIELD_STRING
+	
+	FString ToDelimString(const ECheckBufferFlags InEnum)
+	{		
+		if (const TCHAR* ExactMatch = LexToString(InEnum))
+		{
+			return ExactMatch;
+		}
+		FString String;
+		FOREACH_ECHECKBUFFERFLAGS(TEST_AND_BUILD_BITFIELD_STRING)
+		return String;
+	}
+	
+	// Clean up macros.
+	#undef TEST_AND_BUILD_BITFIELD_STRING
+	#undef TEST_STRING_AGAINST_ENUM_AND_RETURN
+	#undef TEST_CASE_ENUM_RETURN_SHORTNAME
+	#undef FOREACH_ECHECKBUFFERFLAGS
+	#undef FOREACH_ECHECKBUFFERBEHAVIOR
+		
+	static bool TryParseDelimitedBitfield(const FString& InBitfieldString, Audio::ECheckBufferFlags& OutEnum)
+	{
+		using namespace Audio;
+		TArray<FString> Fields;
+		static const TCHAR* Delims[] = { TEXT("|"),TEXT(","),TEXT(" ") };
+		InBitfieldString.ParseIntoArray(Fields, Delims, UE_ARRAY_COUNT(Delims), true);
+		OutEnum = Algo::TransformAccumulate(
+			Fields,
+			[](const FString& i) -> ECheckBufferFlags
+			{
+				Audio::ECheckBufferFlags Flags = ECheckBufferFlags::None;
+				TryLexFromString(*i, Flags);
+				return Flags;
+			},
+			ECheckBufferFlags::None,
+			[](const ECheckBufferFlags& A, const ECheckBufferFlags& B) -> ECheckBufferFlags
+			{
+				return A | B;
+			});
+		return true;
+	}
+
+	TOptional<FString> FindArg(const TCHAR* InName, const TArray<FString>& InArgs)
+	{
+		if (const FString* Found = Algo::FindByPredicate(InArgs, [=](const FString& InStr){ return InStr.StartsWith(InName); }) )
+		{
+			FString Result;
+			Found->Split(InName, nullptr, &Result);
+			return Result;
+		}
+		return {};
+	}
+
+	void HandleDiagnosticsCmd(const TArray<FString>& InArgs, FOutputDevice& InDevice)
+	{
+		using namespace Audio;
+
+		if (FindArg(TEXT("-Help"), InArgs).IsSet())
+		{
+			InDevice.Logf(TEXT("-Filter=<match> -Behavior=Break|Log|Ensure|Nothing -Check=Nan|Infinity|Denormalized|All -Reset -Verbose -Help"));
+			return;
+		}
+		
+		FString Filter; // Default (no filter).
+		if (TOptional<FString> Arg = FindArg(TEXT("-Filter="),InArgs); Arg.IsSet())
+		{
+			Filter = *Arg;
+		}
+		EBufferCheckBehavior Behavior = EBufferCheckBehavior::Ensure;
+		if (TOptional<FString> Arg = FindArg(TEXT("-Behavior="),InArgs); Arg.IsSet())
+		{
+			TryLexFromString(**Arg, Behavior);
+		}
+		ECheckBufferFlags CheckFlags = ECheckBufferFlags::All;
+		if (TOptional<FString> Arg = FindArg(TEXT("-Checks="),InArgs); Arg.IsSet())
+		{
+			TryParseDelimitedBitfield(*Arg, CheckFlags);
+		}
+		bool bVerbose = false;
+		if (FindArg(TEXT("-Verbose"),InArgs).IsSet())
+		{
+			bVerbose = true;
+		}
+		bool bReset = false;
+		if (FindArg(TEXT("-Reset"),InArgs).IsSet())
+		{
+			bReset = true;
+		}
+		
+		InDevice.Logf(TEXT("Submix Buffer Diagnostics. UsingFlags='%s', Behavior='%s'"),
+			ToCStr(ToDelimitedString(CheckFlags)), ToCStr(LexToString(Behavior)) );
+
+		int32 Total = 0;
+		int32 FilterCount = 0;
+		FCheckedBufferState::ForEach([&](Audio::FCheckedBufferState& i)
+		{
+			Total++;
+			if (FString(i.Name).Contains(Filter) || Filter.IsEmpty())
+			{
+				i.Behavior = Behavior;
+				i.CheckFlags = CheckFlags;
+				
+				if (bReset)
+				{
+					i.FailFlags = ECheckBufferFlags::None;
+				}
+				
+				FilterCount++;
+				if (bVerbose)
+				{
+					InDevice.Logf(TEXT("Check='%s', File='%s:%d', Flags='%s', Behavior='%s'"),
+						i.Name,i.File,i.Line, *ToDelimitedString(i.CheckFlags), ToCStr(LexToString(Behavior)) );
+				}
+			}
+		});
+		InDevice.Logf(TEXT("TotalChecks=%d, Filtered=%d"), Total, FilterCount);
+	}
+}//namespace BufferDiagnosticsPrivate
+
+namespace Audio
+{
+	FString ToDelimitedString(const ECheckBufferFlags InEnum)
+	{
+		return BufferDiagnosticsPrivate::ToDelimString(InEnum);
+	}
+}
+	
\ No newline at end of file
diff --git a/Engine/Source/Runtime/SignalProcessing/Private/CheckedAudioBuffer.cpp b/Engine/Source/Runtime/SignalProcessing/Private/CheckedAudioBuffer.cpp
new file mode 100644
--- /dev/null
+++ b/Engine/Source/Runtime/SignalProcessing/Private/CheckedAudioBuffer.cpp
@@ -0,0 +1,135 @@
+// Copyright Epic Games, Inc. All Rights Reserved.
+#include "CheckedAudioBuffer.h"
+#include "CoreGlobals.h"
+
+// This defaults to 'on' unless it's defined ahead.
+#ifndef WITH_AUDIO_BUFFER_CHECKS
+	#define WITH_AUDIO_BUFFER_CHECKS (1)
+#endif //WITH_AUDIO_BUFFER_CHECKS
+
+namespace Audio
+{
+	void FCheckedAudioBuffer::operator=(const FAlignedFloatBuffer& InOther)
+	{
+		DoCheck(Buffer);
+		DoCheck(InOther);
+		Buffer = InOther;
+	}
+	
+	float* FCheckedAudioBuffer::GetData()
+	{
+		DoCheck(Buffer);
+		return Buffer.GetData();
+	}
+
+	int32 FCheckedAudioBuffer::Num() const
+	{
+		DoCheck(Buffer);
+		return Buffer.Num();
+	}
+
+	void FCheckedAudioBuffer::Reserve(const int32 InSize)
+	{
+		Buffer.Reserve(InSize);
+	}
+
+	void FCheckedAudioBuffer::Reset(const int32 InSize/*=0*/)
+	{
+		Buffer.Reset(InSize);
+		bFailedChecks = false; // Might as well reset here.
+		FailedFlags = ECheckBufferFlags::None; 
+	}
+
+	void FCheckedAudioBuffer::AddZeroed(const int32 InSize)
+	{
+		DoCheck(Buffer);
+		Buffer.AddZeroed(InSize);
+	}
+
+	void FCheckedAudioBuffer::SetNumZeroed(const int32 InSize)
+	{
+		DoCheck(Buffer);
+		Buffer.SetNumZeroed(InSize);
+	}
+
+	void FCheckedAudioBuffer::SetNumUninitialized(const int32 InNum)
+	{
+		return Buffer.SetNumUninitialized(InNum);
+	}
+
+	Audio::FAlignedFloatBuffer& FCheckedAudioBuffer::GetBuffer()
+	{
+		DoCheck(Buffer);
+		return Buffer;
+	}
+
+	const Audio::FAlignedFloatBuffer& FCheckedAudioBuffer::GetBuffer() const
+	{
+		DoCheck(Buffer);
+		return Buffer;
+	}
+
+	void FCheckedAudioBuffer::Append(const FCheckedAudioBuffer& InBuffer)
+	{
+		DoCheck(InBuffer );
+		DoCheck(Buffer);
+		Buffer.Append(InBuffer.GetBuffer());
+	}
+
+	void FCheckedAudioBuffer::Append(TArrayView<const float> InView)
+	{
+		DoCheck(InView);
+		DoCheck(Buffer);
+		Buffer.Append(InView);
+	}
+
+	void FCheckedAudioBuffer::Append(const FAlignedFloatBuffer& InBuffer)
+	{
+		DoCheck(InBuffer);
+		DoCheck(Buffer);
+		Buffer.Append(InBuffer);
+	}
+	
+	void FCheckedAudioBuffer::DoCheck(TArrayView<const float> InBuffer) const
+	{
+#if WITH_AUDIO_BUFFER_CHECKS
+		// Only do check until we fail.
+		if (!bFailedChecks && CheckFlags != ECheckBufferFlags::None)
+		{
+			bFailedChecks = CheckBuffer(InBuffer, CheckFlags, FailedFlags);
+			if (bFailedChecks)
+			{
+				switch (Behavior)
+				{
+				// Do nothing.
+				case ECheckBehavior::Ensure:
+					{
+						ensureMsgf(false, TEXT("Buffer Check Ensure. CheckFlags='%s', Failed='%s', Buffer=%s"), 
+							*ToDelimitedString(CheckFlags), *ToDelimitedString(FailedFlags), *DescriptiveName);
+						break;
+					}
+				case ECheckBehavior::Log:
+					{
+						UE_LOG(LogInit, Warning, TEXT("Buffer Check Failed. CheckFlags='%s', Failed='%s', Buffer=%s"), 
+							*ToDelimitedString(CheckFlags), *ToDelimitedString(FailedFlags), *DescriptiveName);
+						break;
+					}
+				case ECheckBehavior::Break:
+					{
+						UE_LOG(LogInit, Warning, TEXT("Buffer Check Failed. CheckFlags='%s', Failed='%s', Buffer=%s"),
+							*ToDelimitedString(CheckFlags), *ToDelimitedString(FailedFlags), *DescriptiveName);
+
+						UE_DEBUG_BREAK_AND_PROMPT_FOR_REMOTE();
+						break;
+					}
+				case ECheckBehavior::Nothing:
+				default:
+					{
+						break;
+					}
+				}
+			}
+		}
+#endif //WITH_AUDIO_BUFFER_CHECKS
+	}
+}
diff --git a/Engine/Source/Runtime/SignalProcessing/Public/CheckedAudioBuffer.h b/Engine/Source/Runtime/SignalProcessing/Public/CheckedAudioBuffer.h
new file mode 100644
--- /dev/null
+++ b/Engine/Source/Runtime/SignalProcessing/Public/CheckedAudioBuffer.h
@@ -0,0 +1,62 @@
+// Copyright Epic Games, Inc. All Rights Reserved.
+#pragma once
+
+#include "DSP/AlignedBuffer.h"
+#include "DSP/BufferDiagnostics.h"
+
+namespace Audio
+{
+	/**
+	 ** Opaque wrapper around FAlignedFloatBuffer.
+	 ** Will perform different checks on accessing the buffer, which can be enabled by flags.
+	 ** Error states are sticky.
+	 **/
+	class SIGNALPROCESSING_API FCheckedAudioBuffer
+	{
+	public:
+		using ECheckBehavior = Audio::EBufferCheckBehavior;
+				
+		void SetName(const FString& InName) { DescriptiveName = InName; }
+		void SetCheckBehavior(const ECheckBehavior InBehavior) { Behavior = InBehavior; }
+		void SetCheckFlags(const ECheckBufferFlags InCheckFlags) { CheckFlags = InCheckFlags; }
+		
+		// Automatic conversions to FAlignedBuffer with a Check.
+		operator FAlignedFloatBuffer&() { return GetBuffer(); };
+		operator TArrayView<float>() { return MakeArrayView(GetBuffer()); };
+		operator TArrayView<const float>() const { return MakeArrayView(GetBuffer()); };
+		operator const FAlignedFloatBuffer&() const { return GetBuffer(); }
+		void operator=(const FAlignedFloatBuffer& InOther);
+		const FAlignedFloatBuffer* operator&() const { return &GetBuffer(); }
+		FAlignedFloatBuffer* operator&() { return &GetBuffer(); }
+				
+		int32 Num() const;
+		void Reserve(const int32 InSize);
+		void Reset(const int32 InSize=0);
+		void AddZeroed(const int32 InSize);
+		void SetNumZeroed(const int32 InSize);
+		void SetNumUninitialized(const int32 InNum);
+
+		// Const.
+		const FAlignedFloatBuffer& GetBuffer() const;
+
+		// Non-const.
+		FAlignedFloatBuffer& GetBuffer();
+		float* GetData();
+				
+		void Append(const FAlignedFloatBuffer& InBuffer);
+		void Append(TArrayView<const float> InView);
+		void Append(const FCheckedAudioBuffer& InBuffer);
+	private:
+		void DoCheck(TArrayView<const float> InBuffer) const;
+		
+		FString DescriptiveName;			// String so this can be procedural.
+		FAlignedFloatBuffer Buffer;			// Wrapped buffer.
+		
+		ECheckBehavior Behavior = ECheckBehavior::Nothing;
+		ECheckBufferFlags CheckFlags = ECheckBufferFlags::None;
+		mutable ECheckBufferFlags FailedFlags = ECheckBufferFlags::None;
+		mutable bool bFailedChecks = false;
+	};
+}
+
+
diff --git a/Engine/Source/Runtime/SignalProcessing/Public/DSP/BufferDiagnostics.h b/Engine/Source/Runtime/SignalProcessing/Public/DSP/BufferDiagnostics.h
new file mode 100644
--- /dev/null
+++ b/Engine/Source/Runtime/SignalProcessing/Public/DSP/BufferDiagnostics.h
@@ -0,0 +1,155 @@
+// Copyright Epic Games, Inc. All Rights Reserved.
+
+#pragma once
+
+#include "Containers/ArrayView.h"
+#include "Containers/List.h"
+#include "Containers/UnrealString.h"
+#include "HAL/CriticalSection.h"
+#include "HAL/PreprocessorHelpers.h"
+#include "Math/Float32.h"
+#include "Templates/Function.h"
+
+namespace Audio
+{
+	// IEEE Standard 754 Floating Point Numbers
+	// https://www.geeksforgeeks.org/ieee-standard-754-floating-point-numbers/
+
+	/*
+	* Denormalized.
+	*/
+	FORCEINLINE static bool IsDenormalized(const float InValue)
+	{
+		// If the exponent is all zeros, but the mantissa is not then the value is a denormalized number.
+		// This means this number does not have an assumed leading one before the binary point.
+		const FFloat32 F(InValue);
+		return F.Components.Exponent == 0 && F.Components.Mantissa != 0;
+	}
+
+	SIGNALPROCESSING_API int32 FindDenormalized(TArrayView<const float> InBuffer);
+
+	SIGNALPROCESSING_API bool ContainsDenormalized(TArrayView<const float> InBuffer);
+
+	/*
+	 * Infinity.
+	*/
+	FORCEINLINE static bool IsInfinity(const float InValue)
+	{
+		// The values + infinity and -infinity are denoted with an exponent of all ones and a mantissa of all zeros.
+		// The sign bit distinguishes between negative infinity and positive infinity.
+		const FFloat32 F(InValue);
+		return F.Components.Exponent == 255 && F.Components.Mantissa == 0;
+	}
+
+	SIGNALPROCESSING_API int32 FindInfinity(TArrayView<const float> InBuffer);
+	SIGNALPROCESSING_API bool ContainsInfinity(TArrayView<const float> InBuffer);
+
+	/*
+	* NaN (Not a number)
+	*/
+	SIGNALPROCESSING_API int32 FindNan(TArrayView<const float> InBuffer);
+
+	SIGNALPROCESSING_API bool ContainsNan(TArrayView<const float> InBuffer);
+
+	enum class SIGNALPROCESSING_API ECheckBufferFlags : uint32
+	{
+		None = 0,
+		Infinity = 1 << 1,
+		Nan = 1 << 2,
+		Denormalized = 1 << 3,
+		All = Infinity | Nan | Denormalized
+	};
+
+	inline ECheckBufferFlags operator|(const ECheckBufferFlags A, const ECheckBufferFlags B)
+	{
+		return static_cast<ECheckBufferFlags>(static_cast<uint32>(A) | static_cast<uint32>(B));
+	}
+
+	inline ECheckBufferFlags& operator|=(ECheckBufferFlags& Out, const ECheckBufferFlags Other)
+	{
+		return Out = Out | Other;
+	}
+
+	inline ECheckBufferFlags operator&(const ECheckBufferFlags A, const ECheckBufferFlags B)
+	{
+		return static_cast<ECheckBufferFlags>(static_cast<uint32>(A) & static_cast<uint32>(B));
+	}
+
+	// ToString with deliminators i.e. "Nan|Infinity" etc.
+	FString ToDelimitedString(const ECheckBufferFlags InEnum);
+
+	/**
+	 * Performs various tests on a audio buffer.
+	 * @param InBuffer Buffer to examine
+	 * @parm InFlags Bitfield of checks to perform
+	 * @return true if check passes, false otherwise
+	**/
+	SIGNALPROCESSING_API bool CheckBuffer(TArrayView<const float> InBuffer, const ECheckBufferFlags InFlags, ECheckBufferFlags& OutFailedFlags);
+
+	enum class SIGNALPROCESSING_API EBufferCheckBehavior : uint8
+	{
+		Nothing,
+		Ensure,
+		Log,
+		Break,
+	};
+
+	struct FCheckedBufferState
+	{
+		// Statics.
+		static FCriticalSection ListCs;	// To protect linked-list.
+		static FCheckedBufferState* Head;
+		static void ForEach(TFunctionRef<void (FCheckedBufferState&)> InCmd);
+		
+		// Per check state.
+		FCheckedBufferState* Next = nullptr;
+		const TCHAR* Name = nullptr;
+		int32 Line = 0;
+		const TCHAR* File = nullptr;
+		ECheckBufferFlags CheckFlags = ECheckBufferFlags::None;
+		ECheckBufferFlags FailFlags = ECheckBufferFlags::None;
+		EBufferCheckBehavior Behavior = EBufferCheckBehavior::Ensure;
+
+		SIGNALPROCESSING_API FCheckedBufferState(const int32 InLine, const TCHAR* InFile, const TCHAR* InName,
+			const ECheckBufferFlags InCheckFlags = ECheckBufferFlags::None,
+			const EBufferCheckBehavior InCheckBehavior = EBufferCheckBehavior::Ensure);
+
+		// The CHECK_AUDIO_BUFFER Macros calls these two member functions.
+		SIGNALPROCESSING_API bool DoCheck(TArrayView<const float> InBuffer);
+		SIGNALPROCESSING_API void FailedBufferCheckImpl(const TCHAR* InFormat, ...) const;
+	};
+
+	
+}
+
+#ifndef WITH_AUDIO_BUFFERDIAGNOSTICS
+// By default enabled in all but ship
+#define WITH_AUDIO_BUFFERDIAGNOSTICS (!UE_BUILD_SHIPPING)
+#endif //WITH_AUDIO_BUFFERDIAGNOSTICS
+
+#if WITH_AUDIO_BUFFERDIAGNOSTICS
+#define AUDIO_CHECK_BUFFER(BUFFER)\
+	static FCheckedBufferState PREPROCESSOR_JOIN(BufferCheck, __LINE__)(__LINE__,TEXT(__FILE__),TEXT(PREPROCESSOR_TO_STRING(BUFFER)));\
+	if (!PREPROCESSOR_JOIN(BufferCheck, __LINE__).DoCheck(BUFFER) )\
+	{\
+		PREPROCESSOR_JOIN(BufferCheck, __LINE__).FailedBufferCheckImpl(TEXT(""));\
+	}
+
+#define AUDIO_CHECK_BUFFER_NAMED(BUFFER,NAME)\
+	static FCheckedBufferState PREPROCESSOR_JOIN(BufferCheck, __LINE__)(__LINE__,TEXT(__FILE__),NAME);\
+	if (!PREPROCESSOR_JOIN(BufferCheck, __LINE__).DoCheck(BUFFER) )\
+	{\
+		PREPROCESSOR_JOIN(BufferCheck, __LINE__).FailedBufferCheckImpl(TEXT(""));\
+	}
+
+#define AUDIO_CHECK_BUFFER_NAMED_MSG(BUFFER,NAME,MSG,...)\
+	static FCheckedBufferState PREPROCESSOR_JOIN(BufferCheck, __LINE__)(__LINE__,TEXT(__FILE__),NAME);\
+	if (!PREPROCESSOR_JOIN(BufferCheck, __LINE__).DoCheck(BUFFER))\
+	{\
+		PREPROCESSOR_JOIN(BufferCheck, __LINE__).FailedBufferCheckImpl(MSG,__VA_ARGS__);\
+	}
+#else //WITH_AUDIO_BUFFERDIAGNOSTICS
+	#define AUDIO_CHECK_BUFFER_NAMED(BUFFER,NAME)
+	#define AUDIO_CHECK_BUFFER(BUFFER)
+	#define AUDIO_CHECK_BUFFER_NAMED_MSG(BUFFER,NAME,MSG,...)
+#endif //WITH_AUDIO_BUFFERDIAGNOSTICS
