Biquad filter optimization for mono process audio

diff --git a/Engine/Source/Runtime/SignalProcessing/Private/Filter.cpp b/Engine/Source/Runtime/SignalProcessing/Private/Filter.cpp
--- a/Engine/Source/Runtime/SignalProcessing/Private/Filter.cpp
+++ b/Engine/Source/Runtime/SignalProcessing/Private/Filter.cpp
@@ -66,9 +66,19 @@ namespace Audio
 	{
 		if (bEnabled)
 		{
-			for (int32 SampleIndex = 0; SampleIndex < InNumSamples; SampleIndex += NumChannels)
+			if (NumChannels == 1)
+			{
+				for (int32 SampleIndex = 0; SampleIndex < InNumSamples; ++SampleIndex)
+				{
+					OutBuffer[SampleIndex] = Biquad->ProcessAudio(InBuffer[SampleIndex]);
+				}
+			}
+			else
 			{
-				ProcessAudioFrame(&InBuffer[SampleIndex], &OutBuffer[SampleIndex]);
+				for (int32 SampleIndex = 0; SampleIndex < InNumSamples; SampleIndex += NumChannels)
+				{
+					ProcessAudioFrame(&InBuffer[SampleIndex], &OutBuffer[SampleIndex]);
+				}
 			}
 		}
 	}
