[metasound] optimization for mixer node (avoid memset / memzero)

diff --git a/Engine/Plugins/Runtime/Metasound/Source/MetasoundStandardNodes/Private/MetasoundMixerNode.cpp b/Engine/Plugins/Runtime/Metasound/Source/MetasoundStandardNodes/Private/MetasoundMixerNode.cpp
--- a/Engine/Plugins/Runtime/Metasound/Source/MetasoundStandardNodes/Private/MetasoundMixerNode.cpp
+++ b/Engine/Plugins/Runtime/Metasound/Source/MetasoundStandardNodes/Private/MetasoundMixerNode.cpp
@@ -214,17 +214,27 @@ namespace Metasound
 
 		void Execute()
 		{
-			// zero the outputs
+			// initialize the output using the first set of input channels
+			float NextGain = *Gains[0];
+			float PrevGain = PrevGains[0];
+
+			// initialize the output buffers w/ the first set of input buffers 
 			for (uint32 i = 0; i < NumChannels; ++i)
 			{
-				Outputs[i]->Zero();
+				TArrayView<const float> InputView = *Inputs[i];
+				TArrayView<float> OutputView = *Outputs[i];
+	
+				check(InputView.Num() == OutputView.Num());
+				Audio::ArrayFade(InputView, PrevGain, NextGain, OutputView);
 			}
 
-			// for each input
-			for (uint32 InputIndex = 0; InputIndex < NumInputs; ++InputIndex)
+			PrevGains[0] = NextGain;
+
+			// mix in each following input
+			for (uint32 InputIndex = 1; InputIndex < NumInputs; ++InputIndex)
 			{
-				const float NextGain = *Gains[InputIndex];
-				const float PrevGain = PrevGains[InputIndex];
+				NextGain = *Gains[InputIndex];
+				PrevGain = PrevGains[InputIndex];
 
 				// for each channel of audio
 				for (uint32 ChanIndex = 0; ChanIndex < NumChannels; ++ChanIndex)
diff --git a/Engine/Source/Runtime/SignalProcessing/Private/FloatArrayMath.cpp b/Engine/Source/Runtime/SignalProcessing/Private/FloatArrayMath.cpp
--- a/Engine/Source/Runtime/SignalProcessing/Private/FloatArrayMath.cpp
+++ b/Engine/Source/Runtime/SignalProcessing/Private/FloatArrayMath.cpp
@@ -2018,6 +2018,80 @@ namespace Audio
 		}
 	}
 
+	void ArrayFade(TArrayView<const float> InBuffer, const float InStartValue, const float InEndValue, TArrayView<float> OutBuffer)
+	{
+		CSV_SCOPED_TIMING_STAT(Audio_Dsp, ArrayFade);
+		
+		const int32 Num = InBuffer.Num();
+		check(Num <= OutBuffer.Num());
+
+		const float* InFloatBuffer = InBuffer.GetData();
+		float* OutFloatBuffer = OutBuffer.GetData();
+
+		// case 1: no fade
+		if (FMath::IsNearlyEqual(InStartValue, InEndValue))
+		{
+			if (InStartValue == 0.0f)
+			{
+				// No need to do anything if start and end values are both 0.0
+				FMemory::Memset(OutFloatBuffer, 0, sizeof(float) * Num);
+			}
+			else
+			{
+				// no fade, just scale the output
+				ArrayMultiplyByConstant(InBuffer, InStartValue, OutBuffer);
+			}
+
+			return;
+		}
+
+		// case 2: fade w/ ISPC
+#if INTEL_ISPC
+		if (bAudio_FloatArrayMath_ISPC_Enabled)
+		{
+			ispc::ArrayFade2(InFloatBuffer, Num, InStartValue, InEndValue, OutFloatBuffer);
+			return;
+		}
+#endif
+
+
+		// case 3: fade w/ our vectorization abstraction
+		const int32 NumToSimd = Num & MathIntrinsics::SimdMask;
+		const int32 NumNotToSimd = Num & MathIntrinsics::NotSimdMask;
+
+		const float DeltaValue = ((InEndValue - InStartValue) / Num);
+
+		if (NumToSimd)
+		{
+			constexpr VectorRegister4Float VectorFour = MakeVectorRegisterFloatConstant(4.f, 4.f, 4.f, 4.f);
+			VectorRegister4Float Accumulator = MakeVectorRegisterFloat(0.f, 1.f, 2.f, 3.f);
+			VectorRegister4Float Delta = VectorLoadFloat1(&DeltaValue);
+			VectorRegister4Float Start = VectorLoadFloat1(&InStartValue);
+
+			for (int32 i = 0; i < NumToSimd; i += AUDIO_NUM_FLOATS_PER_VECTOR_REGISTER)
+			{
+				VectorRegister4Float Gain = VectorMultiplyAdd(Accumulator, Delta, Start);
+				VectorRegister4Float Input = VectorLoad(&InFloatBuffer[i]);
+				VectorRegister4Float Output = VectorMultiply(Input, Gain);
+
+				Accumulator = VectorAdd(Accumulator, VectorFour);
+				VectorStore(Output, &OutFloatBuffer[i]);
+			}
+		}
+
+		if (NumNotToSimd)
+		{
+			float Gain = (NumToSimd * DeltaValue) + InStartValue;
+
+			// Do a fade from start to end
+			for (int32 i = NumToSimd; i < Num; ++i)
+			{
+				OutFloatBuffer[i] = InFloatBuffer[i] * Gain;
+				Gain += DeltaValue;
+			}
+		}
+	}
+
 	void ArrayMixIn(TArrayView<const float> InFloatBuffer, TArrayView<float> BufferToSumTo, const float Gain)
 	{
 		CSV_SCOPED_TIMING_STAT(Audio_Dsp, ArrayMixIn);
diff --git a/Engine/Source/Runtime/SignalProcessing/Private/FloatArrayMath.ispc b/Engine/Source/Runtime/SignalProcessing/Private/FloatArrayMath.ispc
--- a/Engine/Source/Runtime/SignalProcessing/Private/FloatArrayMath.ispc
+++ b/Engine/Source/Runtime/SignalProcessing/Private/FloatArrayMath.ispc
@@ -613,6 +613,21 @@ export void ArrayFade(uniform float OutFloatBuffer[],
 	}
 }
 
+export void ArrayFade2(const uniform float InBuffer[],
+							const uniform int NumSamples,
+							const uniform float StartValue,
+							const uniform float EndValue,
+							uniform float OutBuffer[])
+{
+	const uniform float DeltaValue = ((EndValue - StartValue) / NumSamples);
+
+	foreach(i = 0 ... NumSamples)
+	{
+		const varying float Gain = (i * DeltaValue) + StartValue;
+		OutBuffer[i] = InBuffer[i] * Gain;
+	}
+}
+
 export void ArrayMixInWithGain(const uniform float InFloatBuffer[],
 								uniform float BufferToSumTo[],
 								const uniform int NumSamples,
diff --git a/Engine/Source/Runtime/SignalProcessing/Public/DSP/FloatArrayMath.h b/Engine/Source/Runtime/SignalProcessing/Public/DSP/FloatArrayMath.h
--- a/Engine/Source/Runtime/SignalProcessing/Public/DSP/FloatArrayMath.h
+++ b/Engine/Source/Runtime/SignalProcessing/Public/DSP/FloatArrayMath.h
@@ -221,6 +221,7 @@ namespace Audio
 	/* This operation completely ignores channel counts, so avoid using this function on buffers that are not mono, stereo or quad */
 	/* if the buffer needs to fade all channels uniformly. */
 	SIGNALPROCESSING_API void ArrayFade(TArrayView<float> InOutBuffer, const float StartValue, const float EndValue);
+	SIGNALPROCESSING_API void ArrayFade(TArrayView<const float> InBuffer, const float InStartValue, const float InEndValue, TArrayView<float> OutBuffer);
 
 	/** Takes buffer InFloatBuffer, optionally multiplies it by Gain, and adds it to BufferToSumTo. */
 	SIGNALPROCESSING_API void ArrayMixIn(TArrayView<const float> InFloatBuffer, TArrayView<float> BufferToSumTo, const float Gain);
