Initial implementation of ASIO audio mixer backend.

diff --git a/Engine/Source/Runtime/AudioExtensions/Public/IAudioExtensionPlugin.h b/Engine/Source/Runtime/AudioExtensions/Public/IAudioExtensionPlugin.h
--- a/Engine/Source/Runtime/AudioExtensions/Public/IAudioExtensionPlugin.h
+++ b/Engine/Source/Runtime/AudioExtensions/Public/IAudioExtensionPlugin.h
@@ -347,6 +347,11 @@ public:
 	{
 	}
 
+	virtual void OnInitSource(const uint32 SourceId, const FName& AudioComponentUserId, const uint32 NumChannels, USpatializationPluginSourceSettingsBase* InSettings)
+	{
+		OnInitSource(SourceId, AudioComponentUserId, InSettings);
+	}
+
 	/** Called when a source is done playing and is released. */
 	virtual void OnReleaseSource(const uint32 SourceId)
 	{
diff --git a/Engine/Source/Runtime/AudioMixer/Private/AudioMixerDevice.cpp b/Engine/Source/Runtime/AudioMixer/Private/AudioMixerDevice.cpp
--- a/Engine/Source/Runtime/AudioMixer/Private/AudioMixerDevice.cpp
+++ b/Engine/Source/Runtime/AudioMixer/Private/AudioMixerDevice.cpp
@@ -1290,16 +1290,22 @@ namespace Audio
 
  	FAudioPlatformSettings FMixerDevice::GetPlatformSettings() const
  	{
-		FAudioPlatformSettings Settings = AudioMixerPlatform->GetPlatformSettings();
-
-		const int32 DefaultMaxChannels = GetDefault<UAudioSettings>()->GetHighestMaxChannels();
-		UE_LOG(LogAudioMixer, Display, TEXT("Audio Mixer Platform Settings:"));
-		UE_LOG(LogAudioMixer, Display, TEXT("	Sample Rate:						  %d"), Settings.SampleRate);
-		UE_LOG(LogAudioMixer, Display, TEXT("	Callback Buffer Frame Size Requested: %d"), Settings.CallbackBufferFrameSize);
-		UE_LOG(LogAudioMixer, Display, TEXT("	Callback Buffer Frame Size To Use:	  %d"), AudioMixerPlatform->GetNumFrames(Settings.CallbackBufferFrameSize));
-		UE_LOG(LogAudioMixer, Display, TEXT("	Number of buffers to queue:			  %d"), Settings.NumBuffers);
-		UE_LOG(LogAudioMixer, Display, TEXT("	Max Channels (voices):				  %d"), (Settings.MaxChannels > 0) ? Settings.MaxChannels : DefaultMaxChannels);
-		UE_LOG(LogAudioMixer, Display, TEXT("	Number of Async Source Workers:		  %d"), Settings.NumSourceWorkers);
+		FAudioPlatformSettings
+			Settings;
+
+		if (AudioMixerPlatform)
+		{
+			Settings = AudioMixerPlatform->GetPlatformSettings();
+
+			const int32 DefaultMaxChannels = GetDefault<UAudioSettings>()->GetHighestMaxChannels();
+			UE_LOG(LogAudioMixer, Display, TEXT("Audio Mixer Platform Settings:"));
+			UE_LOG(LogAudioMixer, Display, TEXT("	Sample Rate:						  %d"), Settings.SampleRate);
+			UE_LOG(LogAudioMixer, Display, TEXT("	Callback Buffer Frame Size Requested: %d"), Settings.CallbackBufferFrameSize);
+			UE_LOG(LogAudioMixer, Display, TEXT("	Callback Buffer Frame Size To Use:	  %d"), AudioMixerPlatform->GetNumFrames(Settings.CallbackBufferFrameSize));
+			UE_LOG(LogAudioMixer, Display, TEXT("	Number of buffers to queue:			  %d"), Settings.NumBuffers);
+			UE_LOG(LogAudioMixer, Display, TEXT("	Max Channels (voices):				  %d"), (Settings.MaxChannels > 0) ? Settings.MaxChannels : DefaultMaxChannels);
+			UE_LOG(LogAudioMixer, Display, TEXT("	Number of Async Source Workers:		  %d"), Settings.NumSourceWorkers);
+		}
 
  		return Settings;
  	}
diff --git a/Engine/Source/Runtime/AudioMixer/Private/AudioMixerSource.cpp b/Engine/Source/Runtime/AudioMixer/Private/AudioMixerSource.cpp
--- a/Engine/Source/Runtime/AudioMixer/Private/AudioMixerSource.cpp
+++ b/Engine/Source/Runtime/AudioMixer/Private/AudioMixerSource.cpp
@@ -637,6 +637,13 @@ namespace Audio
 			if (ActiveSound)
 			{
 				InitParams.AudioComponentUserID = WaveInstance->ActiveSound->GetAudioComponentUserID();
+#if AUDIO_MIXER_ENABLE_DEBUG_MODE
+				if (InitParams.AudioComponentUserID.IsNone())
+				{
+					InitParams.AudioComponentUserID = ActiveSound->GetSound()->GetFName();
+
+				}
+#endif // AUDIO_MIXER_ENABLE_DEBUG_MODE
 				InitParams.AudioComponentID = WaveInstance->ActiveSound->GetAudioComponentID();
 			}
 
diff --git a/Engine/Source/Runtime/AudioMixer/Private/AudioMixerSourceManager.cpp b/Engine/Source/Runtime/AudioMixer/Private/AudioMixerSourceManager.cpp
--- a/Engine/Source/Runtime/AudioMixer/Private/AudioMixerSourceManager.cpp
+++ b/Engine/Source/Runtime/AudioMixer/Private/AudioMixerSourceManager.cpp
@@ -989,7 +989,7 @@ namespace Audio
 					bUsingSpatializationPlugin = true;
 				}
 
-				SpatialInterfaceInfo.SpatializationPlugin->OnInitSource(SourceId, InitParams.AudioComponentUserID, InitParams.SpatializationPluginSettings);
+				SpatialInterfaceInfo.SpatializationPlugin->OnInitSource(SourceId, InitParams.AudioComponentUserID, InitParams.NumInputChannels, InitParams.SpatializationPluginSettings);
 			}
 
 			// Create the occlusion plugin source effect
diff --git a/Engine/Source/Runtime/AudioMixer/Public/AudioMixerDevice.h b/Engine/Source/Runtime/AudioMixer/Public/AudioMixerDevice.h
--- a/Engine/Source/Runtime/AudioMixer/Public/AudioMixerDevice.h
+++ b/Engine/Source/Runtime/AudioMixer/Public/AudioMixerDevice.h
@@ -253,6 +253,7 @@ namespace Audio
 		const FAudioPlatformDeviceInfo& GetPlatformDeviceInfo() const { return PlatformInfo; };
 
 		FORCEINLINE int32 GetNumDeviceChannels() const { return PlatformInfo.NumChannels; }
+		FORCEINLINE int32 GetNumDirectOutChannels() const { return PlatformInfo.NumDirectOutChannels; }
 
 		int32 GetNumOutputFrames() const { return PlatformSettings.CallbackBufferFrameSize; }
 
diff --git a/Engine/Source/Runtime/AudioMixerCore/Public/AudioMixer.h b/Engine/Source/Runtime/AudioMixerCore/Public/AudioMixer.h
--- a/Engine/Source/Runtime/AudioMixerCore/Public/AudioMixer.h
+++ b/Engine/Source/Runtime/AudioMixerCore/Public/AudioMixer.h
@@ -157,6 +157,9 @@ namespace Audio
 		/** The number of channels supported by the audio device */
 		int32 NumChannels;
 
+		/** The number of channels above the base stereo or 7.1 channels supported by the audio device */
+		int32 NumDirectOutChannels;
+
 		/** The sample rate of the audio device */
 		int32 SampleRate;
 
@@ -179,6 +182,7 @@ namespace Audio
 			Name = TEXT("Unknown");
 			DeviceId = TEXT("Unknown");
 			NumChannels = 0;
+			NumDirectOutChannels = 0;
 			SampleRate = 0;
 			Format = EAudioMixerStreamDataFormat::Unknown;
 			OutputChannelArray.Reset();
@@ -508,6 +512,9 @@ namespace Audio
 		/** Submit the given buffer to the platform's output audio device. */
 		virtual void SubmitBuffer(const uint8* Buffer) {};
 
+		/** Submit a buffer that is to be output directly through a discreet device channel. */
+		virtual void SubmitDirectOutBuffer(const int32 InDirectOutIndex, const Audio::FAlignedFloatBuffer& InBuffer) {};
+
 		/** Allows platforms to filter the requested number of frames to render. Some platforms only support specific frame counts. */
 		virtual int32 GetNumFrames(const int32 InNumReqestedFrames) { return InNumReqestedFrames; }
 
@@ -747,6 +754,8 @@ public:
 
 	/** Creates a new instance of the audio device implemented by the module. */
 	virtual bool IsAudioMixerModule() const { return false; }
+	/** Does this class of device support multiclient access to the driver */
+	virtual bool IsAudioDeviceClassMulticlient() const { return true; }
 	virtual FAudioDevice* CreateAudioDevice() { return nullptr; }
 	virtual Audio::IAudioMixerPlatformInterface* CreateAudioMixerPlatformInterface() { return nullptr; }
 };
diff --git a/Engine/Source/Runtime/Core/Public/Math/UnrealMathSSE.h b/Engine/Source/Runtime/Core/Public/Math/UnrealMathSSE.h
--- a/Engine/Source/Runtime/Core/Public/Math/UnrealMathSSE.h
+++ b/Engine/Source/Runtime/Core/Public/Math/UnrealMathSSE.h
@@ -3464,7 +3464,19 @@ FORCEINLINE VectorRegister4Int VectorFloatToInt(const VectorRegister4Float& A)
 // TODO: LWC: potential loss of data
 FORCEINLINE VectorRegister4Int VectorFloatToInt(const VectorRegister4Double& A)
 {
-	return VectorFloatToInt( MakeVectorRegisterFloatFromDouble(A) );
+	return VectorFloatToInt(MakeVectorRegisterFloatFromDouble(A));
+}
+
+FORCEINLINE VectorRegister4Int VectorDoubleToInt(const VectorRegister4Double& Vec)
+{
+	VectorRegister4Int A = _mm_cvttpd_epi32(Vec.GetXY());
+	VectorRegister4Int B = _mm_cvttpd_epi32(Vec.GetZW());
+	return _mm_xor_si128(A, _mm_shuffle_epi32(B, SHUFFLEMASK(2, 3, 0, 1)));
+}
+
+FORCEINLINE VectorRegister4Int VectorShuffleByte4(const VectorRegister4Int& Vec, const VectorRegister4Int& Mask)
+{
+	return _mm_shuffle_epi8(Vec, Mask);
 }
 
 
@@ -3519,6 +3531,8 @@ FORCEINLINE VectorRegister4Int VectorFloatToInt(const VectorRegister4Double& A)
 #define VectorShiftRightImmLogical(Vec, ImmAmt)     _mm_srli_epi32(Vec, ImmAmt)
 #define VectorCastIntToFloat(Vec)                   _mm_castsi128_ps(Vec)
 #define VectorCastFloatToInt(Vec)                   _mm_castps_si128(Vec)
+#define VectorCastDoubleToInt(Vec)                  _mm_castpd_si128(Vec)
+#define VectorCastIntToDouble(Vec)                  _mm_castsi128_pd(Vec)
 #define VectorShuffleImmediate(Vec, I0, I1, I2, I3) _mm_shuffle_epi32(Vec, _MM_SHUFFLE(I0, I1, I2, I3))
 #define VectorIntExpandLow16To32(V0)				_mm_unpacklo_epi16(V0, _mm_setzero_si128())
 
diff --git a/Engine/Source/Runtime/Engine/Private/AudioDeviceManager.cpp b/Engine/Source/Runtime/Engine/Private/AudioDeviceManager.cpp
--- a/Engine/Source/Runtime/Engine/Private/AudioDeviceManager.cpp
+++ b/Engine/Source/Runtime/Engine/Private/AudioDeviceManager.cpp
@@ -281,7 +281,10 @@ FAudioDeviceParams FAudioDeviceManager::GetDefaultParamsForNewWorld()
 FAudioDeviceHandle FAudioDeviceManager::RequestAudioDevice(const FAudioDeviceParams& InParams)
 {
 	FScopeLock ScopeLock(&DeviceMapCriticalSection);
-	if (InParams.Scope == EAudioDeviceScope::Unique)
+	
+	// If the device class is not multiclient capable then fall back to sharing the device.
+	// Note that this ignores the bCreateNewAudioDeviceForPlayInEditor editor pref.
+	if (InParams.Scope == EAudioDeviceScope::Unique && AudioDeviceModule->IsAudioDeviceClassMulticlient())
 	{
 		return CreateNewDevice(InParams);
 	}
diff --git a/Engine/Source/Runtime/SignalProcessing/Private/FloatArrayMath.cpp b/Engine/Source/Runtime/SignalProcessing/Private/FloatArrayMath.cpp
--- a/Engine/Source/Runtime/SignalProcessing/Private/FloatArrayMath.cpp
+++ b/Engine/Source/Runtime/SignalProcessing/Private/FloatArrayMath.cpp
@@ -2382,6 +2382,161 @@ namespace Audio
 		}
 	}
 	
+	constexpr int CreateByteMask(uint32 A, uint32 B, uint32 C, uint32 D)
+	{
+		return A | (B << 8) | (C << 16) | (D << 24);
+	}
+
+	void ArrayFloatToPcm24(TArrayView<const float> InView, TArrayView<int8> OutView)
+	{
+		CSV_SCOPED_TIMING_STAT(Audio_Dsp, ArrayFloatToPcm24);
+
+		constexpr int32 SizeofPCM24 = 3;
+		constexpr int32 SimdIndexStride = (4 * SizeofPCM24);
+		check((OutView.Num() / SizeofPCM24) >= InView.Num());
+
+		const int32 Num = InView.Num();
+
+		const float* InputPtr = InView.GetData();
+		int8* OutPtr = OutView.GetData();
+
+		// Output byte order must adhere to host endianess
+		const bool bIsOutputLittleEndian = IsHostLittleEndian();
+
+		constexpr float ConversionValue = static_cast<float>(0x7fffff);
+		const VectorRegister4Float Multiplier = VectorSetFloat1(ConversionValue);
+
+		int32 InIndex = 0;
+		int32 OutIndex = 0;
+
+		// In the SIMD loop below, we overwrite an extra 4 bytes of zeros. The ZeroFillMargin
+		// is used to ensure we don't write off the end of the output array.
+		const int32 ZeroFillMargin = 2;
+		const int32 SimdNum = (Num - ZeroFillMargin) & MathIntrinsics::SimdMask;
+
+		// 0x80 designates zero fill for the _mm_shuffle_epi8 intrinsic
+		constexpr uint32 ZeroFill = 0x80;
+		// Here we convert from 4 32-bit ints to 4 24-bit ints using the mask register below.
+		// Notice that every 4th byte is skipped. This is the uneeded high byte of the 32-bit int.
+		VectorRegister4Int Mask;
+		if (bIsOutputLittleEndian)
+		{
+			Mask = MakeVectorRegisterInt(
+				CreateByteMask(0, 1, 2, 4),
+				CreateByteMask(5, 6, 8, 9),
+				CreateByteMask(10, 12, 13, 14),
+				CreateByteMask(ZeroFill, ZeroFill, ZeroFill, ZeroFill));
+		}
+		else
+		{
+			Mask = MakeVectorRegisterInt(
+				CreateByteMask(2, 1, 0, 6),
+				CreateByteMask(5, 4, 10, 9),
+				CreateByteMask(8, 14, 13, 12),
+				CreateByteMask(ZeroFill, ZeroFill, ZeroFill, ZeroFill));
+		}
+
+		for (; InIndex < SimdNum; InIndex += AUDIO_NUM_FLOATS_PER_VECTOR_REGISTER)
+		{
+			const VectorRegister4Float InVector = VectorLoad(&InputPtr[InIndex]);
+			const VectorRegister4Float ScaledVector = VectorMultiply(InVector, Multiplier);
+			const VectorRegister4Int IntVector = VectorFloatToInt(ScaledVector);
+
+			const VectorRegister4Int OutVector = VectorShuffleByte4(IntVector, Mask);
+			VectorIntStore(OutVector, &OutPtr[OutIndex]);
+			OutIndex += SimdIndexStride;
+		}
+
+		for (; InIndex < Num; InIndex++)
+		{
+			const int32 ConvertedValue = InputPtr[InIndex] * ConversionValue;
+
+			AssignPcm24Value(&OutPtr[OutIndex], ConvertedValue, bIsOutputLittleEndian);
+			OutIndex += SizeofPCM24;
+		}
+	}
+
+	FORCEINLINE void AssignPcm24Value(int8* OutPtr, const int32 InValue, const bool bIsOutputLittleEndian)
+	{
+		uint8* UnsignedOutPtr = (uint8*)OutPtr;
+		const uint32 UnsignedInValue = (const uint32)InValue;
+
+		if (bIsOutputLittleEndian)
+		{
+			UnsignedOutPtr[0] = UnsignedInValue & 0xFF;
+			UnsignedOutPtr[1] = UnsignedInValue >> 8 & 0xFF;
+			UnsignedOutPtr[2] = UnsignedInValue >> 16 & 0xFF;
+		}
+		else
+		{
+			UnsignedOutPtr[2] = UnsignedInValue & 0xFF;
+			UnsignedOutPtr[1] = UnsignedInValue >> 8 & 0xFF;
+			UnsignedOutPtr[0] = UnsignedInValue >> 16 & 0xFF;
+		}
+	}
+
+	void ArrayFloatToPcm32(TArrayView<const float> InView, TArrayView<int32> OutView)
+	{
+		CSV_SCOPED_TIMING_STAT(Audio_Dsp, ArrayFloatToPcm32);
+
+		check(OutView.Num() >= InView.Num());
+
+		const int32 Num = InView.Num();
+
+		const float* InputPtr = InView.GetData();
+		int32* OutPtr = OutView.GetData();
+
+		// Use double precision due to the limitations of 
+		// single precision floats (e.g. values >= 2^24 get rounded)
+		constexpr double ConversionValue = static_cast<double>(TNumericLimits<int32>::Max());
+		const VectorRegister4Double Multiplier = VectorSetFloat1(ConversionValue);
+
+		int32 Index = 0;
+		const int32 SimdNum = Num & MathIntrinsics::SimdMask;
+		for (; Index < SimdNum; Index += AUDIO_NUM_FLOATS_PER_VECTOR_REGISTER)
+		{
+			const VectorRegister4Float InputVectorFloat = VectorLoad(&InputPtr[Index]);
+			const VectorRegister4Double InputVector = MakeVectorRegisterDouble(InputVectorFloat);
+			const VectorRegister4Double ScaledVector = VectorMultiply(InputVector, Multiplier);
+			
+			const VectorRegister4Int OutVector = VectorDoubleToInt(ScaledVector);
+			VectorIntStore(OutVector, &OutPtr[Index]);
+		}
+
+		for (; Index < Num; Index++)
+		{
+			const double Value = InputPtr[Index];
+			OutPtr[Index] = static_cast<int32>(Value * ConversionValue);
+		}
+	}
+
+	void ArrayFloatToPcmDouble(TArrayView<const float> InView, TArrayView<double> OutView)
+	{
+		CSV_SCOPED_TIMING_STAT(Audio_Dsp, ArrayFloatToPcmDouble);
+
+		check(OutView.Num() >= InView.Num());
+
+		const int32 Num = InView.Num();
+
+		const float* InputPtr = InView.GetData();
+		double* OutPtr = OutView.GetData();
+
+		int32 Index = 0;
+		const int32 SimdNum = Num & MathIntrinsics::SimdMask;
+		for (; Index < SimdNum; Index += AUDIO_NUM_FLOATS_PER_VECTOR_REGISTER)
+		{
+			const VectorRegister4Float InputVectorFloat = VectorLoad(&InputPtr[Index]);
+			const VectorRegister4Double OutVector = MakeVectorRegisterDouble(InputVectorFloat);
+
+			VectorStore(OutVector, &OutPtr[Index]);
+		}
+
+		for (; Index < Num; ++Index)
+		{
+			OutPtr[Index] = (double)InputPtr[Index];
+		}
+	}
+
 	void ArrayInterleave(const TArray<FAlignedFloatBuffer>& InBuffers, FAlignedFloatBuffer& OutBuffer)
 	{
 		if(InBuffers.Num() == 0)
@@ -2429,6 +2584,15 @@ namespace Audio
 		
 		const int32 NumFrames = InBuffer.Num() / InChannels;
 
+		ArrayDeinterleave(TArrayView<const float>(InBuffer.GetData(), NumFrames), OutBuffers, InChannels);
+	}
+
+	void ArrayDeinterleave(const TArrayView<const float> InView, TArray<FAlignedFloatBuffer>& OutBuffers, const int32 InChannels)
+	{
+		check(InChannels > 0);
+		
+		const int32 NumFrames = InView.Num();
+
 		TArray<float*> BufferPtrArray;
 		BufferPtrArray.Reset(InChannels);
 
@@ -2444,7 +2608,7 @@ namespace Audio
 		
 		float** OutBufferPtr = BufferPtrArray.GetData();
 		
-		ArrayDeinterleave(InBuffer.GetData(), OutBufferPtr, NumFrames, InChannels);
+		ArrayDeinterleave(InView.GetData(), OutBufferPtr, NumFrames, InChannels);
 	}
 
 	void ArrayDeinterleave(const float* RESTRICT InBuffer, float* const* RESTRICT OutBuffers, const int32 InFrames, const int32 InChannels)
@@ -2547,6 +2711,150 @@ namespace Audio
 		}
 	}
 
+	void ArrayInt16SwapBytes(TArrayView<int16> InView)
+	{
+		CSV_SCOPED_TIMING_STAT(Audio_Dsp, ArrayInt16SwapBytes);
+
+		const int32 Num = InView.Num();
+		int16* InputPtr = InView.GetData();
+
+		int32 Index = 0;
+		const int32 SimdNum = Num & MathIntrinsics::Simd8Mask;
+		constexpr int32 NumInt16PerVectorRegister = sizeof(VectorRegister4Int) / sizeof(int16);
+
+		const VectorRegister4Int LeftMask = VectorIntSet1(0x00ff00ff);
+		const VectorRegister4Int RightMask = VectorIntSet1(0xff00ff00);
+
+		for (; Index < SimdNum; Index += NumInt16PerVectorRegister)
+		{
+			VectorRegister4Int InputVector = VectorIntLoad(&InputPtr[Index]);
+
+			const VectorRegister4Int LeftVector = VectorShiftLeftImm(VectorIntAnd(InputVector, LeftMask), 8);
+			const VectorRegister4Int RightVector = VectorShiftRightImmLogical(VectorIntAnd(InputVector, RightMask), 8);
+
+			InputVector = VectorIntOr(LeftVector, RightVector);
+
+			VectorIntStore(InputVector, &InputPtr[Index]);
+		}
+
+		for (; Index < Num; Index++)
+		{
+			InputPtr[Index] = BYTESWAP_ORDER16(InputPtr[Index]);
+		}
+	}
+
+	void ArrayInt24SwapBytes(TArrayView<int8> InView)
+	{
+		CSV_SCOPED_TIMING_STAT(Audio_Dsp, ArrayInt24SwapBytes);
+
+		constexpr int32 SizeofPCM24 = 3;
+		// We can fit 5 PCM24 samples into a vector register with one byte left over
+		constexpr int32 NumPCM24PerVectorRegister = sizeof(VectorRegister4Int) / SizeofPCM24;
+		constexpr int32 SimdIndexStride = (NumPCM24PerVectorRegister * SizeofPCM24);
+		const int32 NumBytes = InView.Num();
+		int8* InputPtr = InView.GetData();
+
+		int32 Index = 0;
+		const int32 MaxSimdBytes = NumBytes - sizeof(VectorRegister4Int);
+
+		VectorRegister4Int Mask = MakeVectorRegisterInt(
+			CreateByteMask(2, 1, 0, 5),
+			CreateByteMask(4, 3, 8, 7),
+			CreateByteMask(6, 11, 10, 9),
+			CreateByteMask(14, 13, 12, 15));
+
+		for (; Index < MaxSimdBytes; Index += SimdIndexStride)
+		{
+			const VectorRegister4Int InputVector = VectorIntLoad(&InputPtr[Index]);
+			const VectorRegister4Int OutVector = VectorShuffleByte4(InputVector, Mask);
+
+			VectorIntStore(OutVector, &InputPtr[Index]);
+		}
+
+		for (; Index < NumBytes; Index += SizeofPCM24)
+		{
+			int8 TempValue = InputPtr[Index];
+
+			InputPtr[Index] = InputPtr[Index + 2];
+			InputPtr[Index + 2] = TempValue;
+		}
+	}
+
+	void ArrayInt32SwapBytes(TArrayView<int32> InView)
+	{
+		CSV_SCOPED_TIMING_STAT(Audio_Dsp, ArrayInt32SwapBytes);
+
+		const int32 Num = InView.Num();
+		int32* InputPtr = InView.GetData();
+
+		int32 Index = 0;
+		const int32 SimdNum = Num & MathIntrinsics::SimdMask;
+		constexpr int32 NumInt32PerVectorRegister = sizeof(VectorRegister4Int) / sizeof(int32);
+
+		VectorRegister4Int Mask = MakeVectorRegisterInt(
+			CreateByteMask(3, 2, 1, 0),
+			CreateByteMask(7, 6, 5, 4),
+			CreateByteMask(11, 10, 9, 8),
+			CreateByteMask(15, 14, 13, 12));
+
+		for (; Index < SimdNum; Index += NumInt32PerVectorRegister)
+		{
+			const VectorRegister4Int InputVector = VectorIntLoad(&InputPtr[Index]);
+
+			// Byte shuffle is approximately 2x faster than mask and shift method in this case
+			const VectorRegister4Int OutVector = VectorShuffleByte4(InputVector, Mask);
+
+			VectorIntStore(OutVector, &InputPtr[Index]);
+		}
+
+		for (; Index < Num; Index++)
+		{
+			InputPtr[Index] = BYTESWAP_ORDER32(InputPtr[Index]);
+		}
+	}
+
+	void ArrayFloatSwapBytes(TArrayView<float> InView)
+	{
+		CSV_SCOPED_TIMING_STAT(Audio_Dsp, ArrayFloatSwapBytes);
+
+		ArrayInt32SwapBytes(TArrayView<int32>((int32*)InView.GetData(), InView.Num()));
+	}
+
+	void ArrayDoubleSwapBytes(TArrayView<double> InView)
+	{
+		CSV_SCOPED_TIMING_STAT(Audio_Dsp, ArrayDoubleSwapBytes);
+
+		const int32 Num = InView.Num();
+		double* InputPtr = InView.GetData();
+
+		int32 Index = 0;
+		const int32 SimdNum = Num & MathIntrinsics::SimdMask;
+		constexpr int32 NumDoublePerVectorRegister = sizeof(VectorRegister4Double) / sizeof(double);
+
+		VectorRegister4Int Mask = MakeVectorRegisterInt(
+			CreateByteMask(7, 6, 5, 4),
+			CreateByteMask(3, 2, 1, 0),
+			CreateByteMask(15, 14, 13, 12),
+			CreateByteMask(11, 10, 9, 8));
+
+		for (; Index < SimdNum; Index += NumDoublePerVectorRegister)
+		{
+			const VectorRegister4Double InputVector = VectorLoad(&InputPtr[Index]);
+
+			const VectorRegister4Int OutVectorXY = VectorShuffleByte4(VectorCastDoubleToInt(InputVector.GetXY()), Mask);
+			const VectorRegister4Int OutVectorZW = VectorShuffleByte4(VectorCastDoubleToInt(InputVector.GetZW()), Mask);
+
+			const VectorRegister4Double OutVector(VectorCastIntToDouble(OutVectorXY), VectorCastIntToDouble(OutVectorZW));
+
+			VectorStore(OutVector, &InputPtr[Index]);
+		}
+
+		for (; Index < Num; Index++)
+		{
+			InputPtr[Index] = BYTESWAP_ORDERD(InputPtr[Index]);
+		}
+	}
+
 	FContiguousSparse2DKernelTransform::FContiguousSparse2DKernelTransform(const int32 NumInElements, const int32 NumOutElements)
 	:	NumIn(NumInElements)
 	,	NumOut(NumOutElements)
diff --git a/Engine/Source/Runtime/SignalProcessing/Public/DSP/FloatArrayMath.h b/Engine/Source/Runtime/SignalProcessing/Public/DSP/FloatArrayMath.h
--- a/Engine/Source/Runtime/SignalProcessing/Public/DSP/FloatArrayMath.h
+++ b/Engine/Source/Runtime/SignalProcessing/Public/DSP/FloatArrayMath.h
@@ -236,6 +236,15 @@ namespace Audio
 	SIGNALPROCESSING_API void ArrayFloatToPcm16(TArrayView<const float> InView, TArrayView<int16> OutView);
 	SIGNALPROCESSING_API void ArrayPcm16ToFloat(TArrayView<const int16> InView, TArrayView<float> OutView);
 
+	/** Converts float PCM data to 24 bit integer PCM data */
+	SIGNALPROCESSING_API void ArrayFloatToPcm24(TArrayView<const float> InView, TArrayView<int8> OutView);
+	/** Copies 24 bits of a 32 bit integer (assumes data has been clamped to 24-bits) */
+	SIGNALPROCESSING_API void AssignPcm24Value(int8* OutPtr, const int32 InValue, const bool bIsOutputLittleEndian);
+	/** Converts float PCM data to 32 bit integer PCM data */
+	SIGNALPROCESSING_API void ArrayFloatToPcm32(TArrayView<const float> InView, TArrayView<int32> OutView);
+	/** Converts float PCM data to double PCM data */
+	SIGNALPROCESSING_API void ArrayFloatToPcmDouble(TArrayView<const float> InView, TArrayView<double> OutView);
+
 	/** Interleaves samples from an array of input buffers */
 	SIGNALPROCESSING_API void ArrayInterleave(const TArray<FAlignedFloatBuffer>& InBuffers, FAlignedFloatBuffer& OutBuffer);
 
@@ -245,12 +254,40 @@ namespace Audio
 	/** Interleaves samples from an array of input buffers */
 	SIGNALPROCESSING_API void ArrayDeinterleave(const FAlignedFloatBuffer& InBuffer, TArray<FAlignedFloatBuffer>& OutBuffers, const int32 InChannels);
 
+	/** Interleaves samples from an array of input buffers */
+	SIGNALPROCESSING_API void ArrayDeinterleave(const TArrayView<const float> InView, TArray<FAlignedFloatBuffer>& OutBuffers, const int32 InChannels);
+
 	/** Interleaves samples from an array of input buffers */
 	SIGNALPROCESSING_API void ArrayDeinterleave(const float* RESTRICT InBuffer, float* const* RESTRICT OutBuffers, const int32 InFrames, const int32 InChannels);
 
 	/** Interpolates a Mono audio buffer. */
 	SIGNALPROCESSING_API void ArrayInterpolate(const float* RESTRICT InBuffer, float* RESTRICT OutBuffer, const int32 NumInSamples, const int32 NumOutSamples);
 
+	/** Vectorized 16-bit integer byte swapping. */
+	SIGNALPROCESSING_API void ArrayInt16SwapBytes(TArrayView<int16> InView);
+
+	/** Vectorized 24-bit integer byte swapping. */
+	SIGNALPROCESSING_API void ArrayInt24SwapBytes(TArrayView<int8> InView);
+
+	/** Vectorized 32-bit integer byte swapping. */
+	SIGNALPROCESSING_API void ArrayInt32SwapBytes(TArrayView<int32> InView);
+
+	/** Vectorized 32-bit float byte swapping. */
+	SIGNALPROCESSING_API void ArrayFloatSwapBytes(TArrayView<float> InView);
+
+	/** Vectorized 64-bit float byte swapping. */
+	SIGNALPROCESSING_API void ArrayDoubleSwapBytes(TArrayView<double> InView);
+
+	/** Returns true if host has little endian byte ordering */
+	SIGNALPROCESSING_API constexpr bool IsHostLittleEndian()
+	{
+#if PLATFORM_LITTLE_ENDIAN
+		return true;
+#else
+		return false;
+#endif // PLATFORM_LITTLE_ENDIAN
+	}
+
 	/** FContiguousSparse2DKernelTransform
 	 *
 	 *  FContiguousSparse2DKernelTransform applies a matrix transformation to an input array. 
