Enable Neon path for newly supported vectorized routines.

diff --git a/Engine/Source/Runtime/SignalProcessing/Private/FloatArrayMath.cpp b/Engine/Source/Runtime/SignalProcessing/Private/FloatArrayMath.cpp
--- a/Engine/Source/Runtime/SignalProcessing/Private/FloatArrayMath.cpp
+++ b/Engine/Source/Runtime/SignalProcessing/Private/FloatArrayMath.cpp
@@ -2403,7 +2403,6 @@ namespace Audio
         // Output byte order must adhere to host endianess
         const bool bIsOutputLittleEndian = IsHostLittleEndian();
 
-#if PLATFORM_ENABLE_VECTORINTRINSICS && !PLATFORM_ENABLE_VECTORINTRINSICS_NEON
         constexpr int32 SimdIndexStride = (4 * SizeofPCM24);
         check((OutView.Num() / SizeofPCM24) >= InView.Num());
 
@@ -2414,7 +2413,8 @@ namespace Audio
 		const int32 ZeroFillMargin = 2;
 		const int32 SimdNum = (Num - ZeroFillMargin) & MathIntrinsics::SimdMask;
 
-		// 0x80 designates zero fill for the _mm_shuffle_epi8 intrinsic
+		// 0x80 designates zero fill for the _mm_shuffle_epi8 intrinsic for SSE
+		// Neon uses >= number of source bytes, which is 0x10 in this case so 0x80 works for both platforms
 		constexpr uint32 ZeroFill = 0x80;
 		// Here we convert from 4 32-bit ints to 4 24-bit ints using the mask register below.
 		// Notice that every 4th byte is skipped. This is the uneeded high byte of the 32-bit int.
@@ -2446,7 +2446,6 @@ namespace Audio
 			VectorIntStore(OutVector, &OutPtr[OutIndex]);
 			OutIndex += SimdIndexStride;
 		}
-#endif // PLATFORM_ENABLE_VECTORINTRINSICS && !PLATFORM_ENABLE_VECTORINTRINSICS_NEON
         
 		for (; InIndex < Num; InIndex++)
 		{
@@ -2491,7 +2490,6 @@ namespace Audio
         // single precision floats (e.g. values >= 2^24 get rounded)
         constexpr double ConversionValue = static_cast<double>(TNumericLimits<int32>::Max());
 
-#if PLATFORM_ENABLE_VECTORINTRINSICS && !PLATFORM_ENABLE_VECTORINTRINSICS_NEON
 		const VectorRegister4Double Multiplier = VectorSetFloat1(ConversionValue);
 
 		const int32 SimdNum = Num & MathIntrinsics::SimdMask;
@@ -2504,7 +2502,6 @@ namespace Audio
 			const VectorRegister4Int OutVector = VectorDoubleToInt(ScaledVector);
 			VectorIntStore(OutVector, &OutPtr[Index]);
 		}
-#endif // PLATFORM_ENABLE_VECTORINTRINSICS && !PLATFORM_ENABLE_VECTORINTRINSICS_NEON
         
 		for (; Index < Num; Index++)
 		{
@@ -2756,7 +2753,6 @@ namespace Audio
         int32 Index = 0;
         constexpr int32 SizeofPCM24 = 3;
 
-#if PLATFORM_ENABLE_VECTORINTRINSICS && !PLATFORM_ENABLE_VECTORINTRINSICS_NEON
 		// We can fit 5 PCM24 samples into a vector register with one byte left over
 		constexpr int32 NumPCM24PerVectorRegister = sizeof(VectorRegister4Int) / SizeofPCM24;
 		constexpr int32 SimdIndexStride = (NumPCM24PerVectorRegister * SizeofPCM24);
@@ -2776,7 +2772,6 @@ namespace Audio
 
 			VectorIntStore(OutVector, &InputPtr[Index]);
 		}
-#endif //PLATFORM_ENABLE_VECTORINTRINSICS && !PLATFORM_ENABLE_VECTORINTRINSICS_NEON
         
 		for (; Index < NumBytes; Index += SizeofPCM24)
 		{
@@ -2796,7 +2791,6 @@ namespace Audio
 
         int32 Index = 0;
 
-#if PLATFORM_ENABLE_VECTORINTRINSICS && !PLATFORM_ENABLE_VECTORINTRINSICS_NEON
 		const int32 SimdNum = Num & MathIntrinsics::SimdMask;
 		constexpr int32 NumInt32PerVectorRegister = sizeof(VectorRegister4Int) / sizeof(int32);
 
@@ -2815,7 +2809,6 @@ namespace Audio
 
 			VectorIntStore(OutVector, &InputPtr[Index]);
 		}
-#endif // PLATFORM_ENABLE_VECTORINTRINSICS && !PLATFORM_ENABLE_VECTORINTRINSICS_NEON
         
 		for (; Index < Num; Index++)
 		{
@@ -2839,7 +2832,6 @@ namespace Audio
 
         int32 Index = 0;
 
-#if PLATFORM_ENABLE_VECTORINTRINSICS && !PLATFORM_ENABLE_VECTORINTRINSICS_NEON
 		const int32 SimdNum = Num & MathIntrinsics::SimdMask;
 		constexpr int32 NumDoublePerVectorRegister = sizeof(VectorRegister4Double) / sizeof(double);
 
@@ -2860,7 +2852,6 @@ namespace Audio
 
 			VectorStore(OutVector, &InputPtr[Index]);
 		}
-#endif //PLATFORM_ENABLE_VECTORINTRINSICS && !PLATFORM_ENABLE_VECTORINTRINSICS_NEON
         
 		for (; Index < Num; Index++)
 		{
