Removed redundant memcpy in MetaSound math node Audio::TMathOpMultiply<>::Calculate()

diff --git a/Engine/Plugins/Runtime/Metasound/Source/MetasoundStandardNodes/Private/MetasoundMathNodes.cpp b/Engine/Plugins/Runtime/Metasound/Source/MetasoundStandardNodes/Private/MetasoundMathNodes.cpp
--- a/Engine/Plugins/Runtime/Metasound/Source/MetasoundStandardNodes/Private/MetasoundMathNodes.cpp
+++ b/Engine/Plugins/Runtime/Metasound/Source/MetasoundStandardNodes/Private/MetasoundMathNodes.cpp
@@ -1098,9 +1098,28 @@ namespace Metasound
 				return;
 			}
 
-			FMemory::Memcpy(OutResult->GetData(), InPrimaryOperand->GetData(), sizeof(float) * NumSamples);
+			if (InAdditionalOperands.Num() > 0)
+			{
+				const FAudioBuffer& Operand = *InAdditionalOperands[0];
+				if (!ensure(NumSamples == Operand.Num()))
+				{
+					// TODO: Error
+					OutResult->Zero();
+					return;
+				}
 
-			for (int32 i = 0; i < InAdditionalOperands.Num(); ++i)
+				TArrayView<const float> OperandADataView(InPrimaryOperand->GetData(), NumSamples);
+				TArrayView<const float> OperandBDataView(Operand.GetData(), NumSamples);
+				TArrayView<float> OutDataView(OutResult->GetData(), NumSamples);
+
+				Audio::ArrayMultiply(OperandADataView, OperandBDataView, OutDataView);
+			}
+			else
+			{
+				FMemory::Memcpy(OutResult->GetData(), InPrimaryOperand->GetData(), sizeof(float) * NumSamples);
+			}
+
+			for (int32 i = 1; i < InAdditionalOperands.Num(); ++i)
 			{
 				const FAudioBuffer& Operand = *InAdditionalOperands[i];
 				if (!ensure(NumSamples == Operand.Num()))
diff --git a/Engine/Source/Runtime/SignalProcessing/Private/FloatArrayMath.cpp b/Engine/Source/Runtime/SignalProcessing/Private/FloatArrayMath.cpp
--- a/Engine/Source/Runtime/SignalProcessing/Private/FloatArrayMath.cpp
+++ b/Engine/Source/Runtime/SignalProcessing/Private/FloatArrayMath.cpp
@@ -804,6 +804,39 @@ namespace Audio
 		return Max;
 	}
 
+	void ArrayMultiply(TArrayView<const float> InFloatBufferA, TArrayView<const float> InFloatBufferB, TArrayView<float> OutBuffer)
+	{
+		CSV_SCOPED_TIMING_STAT(Audio_Dsp, ArrayMultiply);
+		checkf((InFloatBufferA.Num() == InFloatBufferB.Num()) && (InFloatBufferA.Num() == OutBuffer.Num()), TEXT("Input buffers must be equal length"));
+
+		const int32 Num = InFloatBufferA.Num();
+
+		if (bAudio_FloatArrayMath_ISPC_Enabled)
+		{
+#if INTEL_ISPC
+			ispc::ArrayMultiply(InFloatBufferA.GetData(), InFloatBufferB.GetData(), OutBuffer.GetData(), Num);
+#endif
+		}
+		else
+		{
+			const int32 NumToSimd = Num & MathIntrinsics::SimdMask;
+
+			for (int32 i = 0; i < NumToSimd; i += AUDIO_NUM_FLOATS_PER_VECTOR_REGISTER)
+			{
+				VectorRegister4Float Input1 = VectorLoad(&InFloatBufferA[i]);
+				VectorRegister4Float Input2 = VectorLoad(&InFloatBufferB[i]);
+
+				VectorRegister4Float Output = VectorMultiply(Input1, Input2);
+				VectorStore(Output, &OutBuffer[i]);
+			}
+
+			for (int32 i = NumToSimd; i < Num; ++i)
+			{
+				OutBuffer[i] = InFloatBufferA[i] * InFloatBufferB[i];
+			}
+		}
+	}
+
 	void ArrayMultiplyInPlace(TArrayView<const float> InFloatBuffer, TArrayView<float> BufferToMultiply)
 	{
 		CSV_SCOPED_TIMING_STAT(Audio_Dsp, ArrayMultiplyInPlace);
diff --git a/Engine/Source/Runtime/SignalProcessing/Private/FloatArrayMath.ispc b/Engine/Source/Runtime/SignalProcessing/Private/FloatArrayMath.ispc
--- a/Engine/Source/Runtime/SignalProcessing/Private/FloatArrayMath.ispc
+++ b/Engine/Source/Runtime/SignalProcessing/Private/FloatArrayMath.ispc
@@ -239,6 +239,20 @@ export void ArrayMinMaxNormalize(const uniform float InDataPtr[],
 	}
 }
 
+export void ArrayMultiply(const uniform float InFloatBuffer[],
+						  const uniform float BufferToMultiply[],
+						  uniform float OutDataPtr[],
+						  const uniform int NumSamples)
+{
+	foreach(i = 0 ... NumSamples)
+	{
+		const varying float Input = InFloatBuffer[i];
+		const varying float Output = BufferToMultiply[i];
+
+		OutDataPtr[i] = Input * Output;
+	}
+}
+
 export void ArrayMultiplyInPlace(const uniform float InFloatBuffer[],
 									uniform float BufferToMultiply[],
 									const uniform int NumSamples)
diff --git a/Engine/Source/Runtime/SignalProcessing/Public/DSP/FloatArrayMath.h b/Engine/Source/Runtime/SignalProcessing/Public/DSP/FloatArrayMath.h
--- a/Engine/Source/Runtime/SignalProcessing/Public/DSP/FloatArrayMath.h
+++ b/Engine/Source/Runtime/SignalProcessing/Public/DSP/FloatArrayMath.h
@@ -116,6 +116,9 @@ namespace Audio
 	 */
 	SIGNALPROCESSING_API float ArrayMaxAbsValue(const TArrayView<const float> InView);
 
+	/** Multiply the second buffer by the first buffer. */
+	SIGNALPROCESSING_API void ArrayMultiply(TArrayView<const float> InFloatBufferA, TArrayView<const float> InFloatBufferB, TArrayView<float> OutBuffer);
+
 	/** Multiply the second buffer in place by the first buffer. */
 	SIGNALPROCESSING_API void ArrayMultiplyInPlace(TArrayView<const float> InFloatBuffer, TArrayView<float> BufferToMultiply);
 
