Force denormals to zero for AudioThread and audio jobs

diff --git a/Engine/Plugins/Runtime/AudioWidgets/Source/AudioWidgets/Private/AudioSpectrumAnalyzer.cpp b/Engine/Plugins/Runtime/AudioWidgets/Source/AudioWidgets/Private/AudioSpectrumAnalyzer.cpp
--- a/Engine/Plugins/Runtime/AudioWidgets/Source/AudioWidgets/Private/AudioSpectrumAnalyzer.cpp
+++ b/Engine/Plugins/Runtime/AudioWidgets/Source/AudioWidgets/Private/AudioSpectrumAnalyzer.cpp
@@ -181,8 +181,7 @@ namespace AudioWidgets
 			const float OldValue = ARSmoothedSquaredMagnitudes[Index];
 			const float NewValue = SquaredMagnitudes[Index];
 			const float ARSmootherCoefficient = (NewValue >= OldValue) ? AttackRelease.GetAttackTimeSamples() : AttackRelease.GetReleaseTimeSamples();
-			const float SmoothedValue = FMath::Lerp(NewValue, OldValue, ARSmootherCoefficient);
-			ARSmoothedSquaredMagnitudes[Index] = Audio::UnderflowClamp(SmoothedValue);
+			ARSmoothedSquaredMagnitudes[Index] = FMath::Lerp(NewValue, OldValue, ARSmootherCoefficient);
 		}
 	}
 
diff --git a/Engine/Source/Runtime/AudioMixer/Private/AudioMixerSourceDecode.cpp b/Engine/Source/Runtime/AudioMixer/Private/AudioMixerSourceDecode.cpp
--- a/Engine/Source/Runtime/AudioMixer/Private/AudioMixerSourceDecode.cpp
+++ b/Engine/Source/Runtime/AudioMixer/Private/AudioMixerSourceDecode.cpp
@@ -84,6 +84,8 @@ public:
 
 	void DoWork()
 	{
+		FScopedFTZFloatMode FTZ;
+
 		switch (TaskType)
 		{
 			case EAudioTaskType::Procedural:
diff --git a/Engine/Source/Runtime/AudioMixerCore/Private/AudioMixer.cpp b/Engine/Source/Runtime/AudioMixerCore/Private/AudioMixer.cpp
--- a/Engine/Source/Runtime/AudioMixerCore/Private/AudioMixer.cpp
+++ b/Engine/Source/Runtime/AudioMixerCore/Private/AudioMixer.cpp
@@ -777,6 +777,8 @@ namespace Audio
 	{	
 		LLM_SCOPE(ELLMTag::AudioMixer);
 
+		FScopedFTZFloatMode FTZ;
+
 		uint32 ReturnVal = 0;
 		FMemory::SetupTLSCachesOnCurrentThread();
 
diff --git a/Engine/Source/Runtime/Core/Public/Math/VectorRegister.h b/Engine/Source/Runtime/Core/Public/Math/VectorRegister.h
--- a/Engine/Source/Runtime/Core/Public/Math/VectorRegister.h
+++ b/Engine/Source/Runtime/Core/Public/Math/VectorRegister.h
@@ -123,3 +123,19 @@ namespace GlobalVectorConstants
 	inline constexpr VectorRegister RotationSignificantThreshold = MakeVectorRegisterConstant(1.0 - UE_DELTA*UE_DELTA, 1.0 - UE_DELTA*UE_DELTA, 1.0 - UE_DELTA*UE_DELTA, 1.0 - UE_DELTA*UE_DELTA);
 }
 
+struct FScopedFTZFloatMode
+{
+	FScopedFTZFloatMode()
+	{
+		ControlRegisterState = VectorGetControlRegister();
+		VectorSetControlRegister(ControlRegisterState | VECTOR_DENORMALS_FLUSH_TO_ZERO);
+	}
+
+	~FScopedFTZFloatMode()
+	{
+		VectorSetControlRegister(ControlRegisterState);
+	}
+
+private:
+	uint32 ControlRegisterState;
+};
\ No newline at end of file
diff --git a/Engine/Source/Runtime/SignalProcessing/Private/AllPassFilter.cpp b/Engine/Source/Runtime/SignalProcessing/Private/AllPassFilter.cpp
--- a/Engine/Source/Runtime/SignalProcessing/Private/AllPassFilter.cpp
+++ b/Engine/Source/Runtime/SignalProcessing/Private/AllPassFilter.cpp
@@ -29,9 +29,8 @@ namespace Audio
 		const float Wn = InputSample + G*WnD;
 
 		// form y(n) = -gw(n) + w(n-D)
-		float Yn = -G*Wn + WnD;
+		const float Yn = -G*Wn + WnD;
 
-		Yn = UnderflowClamp(Yn);
 		this->WriteDelayAndInc(Wn);
 		return Yn;
 	}
diff --git a/Engine/Source/Runtime/SignalProcessing/Private/AllPassFractionalDelay.cpp b/Engine/Source/Runtime/SignalProcessing/Private/AllPassFractionalDelay.cpp
--- a/Engine/Source/Runtime/SignalProcessing/Private/AllPassFractionalDelay.cpp
+++ b/Engine/Source/Runtime/SignalProcessing/Private/AllPassFractionalDelay.cpp
@@ -154,9 +154,6 @@ namespace Audio
 				OutSamples[i] = DelayData[DelayPos - 1] + CoefficientsData[i] * DelayData[DelayPos] - CoefficientsData[i] * OutSamples[i - 1];
 			}
 
-			TArrayView<float> OutSamplesView(OutSamples, InNum);
-			ArrayUnderflowClamp(OutSamplesView);
-
 			Z1 = OutSamples[InNum - 1];
 		}
 		
diff --git a/Engine/Source/Runtime/SignalProcessing/Private/BufferOnePoleLPF.cpp b/Engine/Source/Runtime/SignalProcessing/Private/BufferOnePoleLPF.cpp
--- a/Engine/Source/Runtime/SignalProcessing/Private/BufferOnePoleLPF.cpp
+++ b/Engine/Source/Runtime/SignalProcessing/Private/BufferOnePoleLPF.cpp
@@ -60,10 +60,9 @@ namespace Audio
 		OutSampleData[0] = InSampleData[0] * A0 + B1 * Z1;
 		for(Index = 1, DelayIndex = 0; Index < InNum; Index++, DelayIndex++)
 		{
-			//OutSampleData[Index] = UnderflowClamp(InSampleData[Index] * A0 + B1 * OutSampleData[DelayIndex]);
 			OutSampleData[Index] = InSampleData[Index] * A0 + B1 * OutSampleData[DelayIndex];
 		}
-		ArrayUnderflowClamp(OutSamples);
+
 		// Store delay value
 		Z1 = OutSampleData[InNum - 1];
 	}
diff --git a/Engine/Source/Runtime/SignalProcessing/Private/BufferVectorOperations.cpp b/Engine/Source/Runtime/SignalProcessing/Private/BufferVectorOperations.cpp
--- a/Engine/Source/Runtime/SignalProcessing/Private/BufferVectorOperations.cpp
+++ b/Engine/Source/Runtime/SignalProcessing/Private/BufferVectorOperations.cpp
@@ -5,221 +5,6 @@
 
 namespace Audio
 {
-	/* Sets a values to zero if value is denormal. Denormal numbers significantly slow down floating point operations. */
-	void BufferUnderflowClampFast(FAlignedFloatBuffer& InOutBuffer)
-	{
-		ArrayUnderflowClamp(InOutBuffer);
-	}
-	
-	/* Sets a values to zero if value is denormal. Denormal numbers significantly slow down floating point operations. */
-	void BufferUnderflowClampFast(float* RESTRICT InOutBuffer, const int32 InNum)
-	{
-		ArrayUnderflowClamp(MakeArrayView<float>(InOutBuffer, InNum));
-	}
-
-	/* Clamps values in the buffer to be between InMinValue and InMaxValue */
-	void BufferRangeClampFast(FAlignedFloatBuffer& InOutBuffer, float InMinValue, float InMaxValue)
-	{
-		ArrayRangeClamp(InOutBuffer, InMinValue, InMaxValue);
-	}
-
-	/* Clamps values in the buffer to be between InMinValue and InMaxValue */
-	void BufferRangeClampFast(float* RESTRICT InOutBuffer, const int32 InNum, float InMinValue, float InMaxValue)
-	{
-		ArrayRangeClamp(MakeArrayView<float>(InOutBuffer, InNum), InMinValue, InMaxValue);
-	}
-
-	void BufferMultiplyByConstant(const FAlignedFloatBuffer& InFloatBuffer, float InValue, FAlignedFloatBuffer& OutFloatBuffer)
-	{
-		ArrayMultiplyByConstant(InFloatBuffer, InValue, OutFloatBuffer);
-	}
-
-	void BufferMultiplyByConstant(const float* RESTRICT InFloatBuffer, float InValue, float* RESTRICT OutFloatBuffer, const int32 InNumSamples)
-	{
-		ArrayMultiplyByConstant(MakeArrayView<const float>(InFloatBuffer, InNumSamples), InValue, MakeArrayView<float>(OutFloatBuffer, InNumSamples));
-	}
-
-	void MultiplyBufferByConstantInPlace(FAlignedFloatBuffer& InBuffer, float InGain)
-	{
-		ArrayMultiplyByConstantInPlace(InBuffer, InGain);
-	}
-
-	void MultiplyBufferByConstantInPlace(float* RESTRICT InBuffer, int32 NumSamples, float InGain)
-	{
-		ArrayMultiplyByConstantInPlace(MakeArrayView<float>(InBuffer, NumSamples), InGain);
-	}
-
-	// Adds a constant to a buffer (useful for DC offset removal)
-	void AddConstantToBufferInplace(FAlignedFloatBuffer& InBuffer, float InConstant)
-	{
-		ArrayAddConstantInplace(InBuffer, InConstant);
-	}
-
-	void AddConstantToBufferInplace(float* RESTRICT InBuffer, int32 NumSamples, float InConstant)
-	{
-		ArrayAddConstantInplace(MakeArrayView<float>(InBuffer, NumSamples), InConstant);
-	}
-
-	void BufferSetToConstantInplace(FAlignedFloatBuffer& InBuffer, float InConstant)
-	{
-		ArraySetToConstantInplace(InBuffer, InConstant);
-	}
-
-	void BufferSetToConstantInplace(float* RESTRICT InBuffer, int32 NumSamples, float InConstant)
-	{
-		ArraySetToConstantInplace(MakeArrayView<float>(InBuffer, NumSamples), InConstant);
-	}
-
-	/* Performs an element-wise weighted sum OutputBuffer = (InBuffer1 x InGain1) + (InBuffer2 x InGain2) */
-	void BufferWeightedSumFast(const FAlignedFloatBuffer& InBuffer1, float InGain1, const FAlignedFloatBuffer& InBuffer2, float InGain2, FAlignedFloatBuffer& OutBuffer)
-	{
-		ArrayWeightedSum(InBuffer1, InGain1, InBuffer2, InGain2, OutBuffer);
-	}
-
-	/* Performs an element-wise weighted sum OutputBuffer = (InBuffer1 x InGain1) + InBuffer2 */
-	void BufferWeightedSumFast(const FAlignedFloatBuffer& InBuffer1, float InGain1, const FAlignedFloatBuffer& InBuffer2, FAlignedFloatBuffer& OutBuffer)
-	{
-		ArrayWeightedSum(InBuffer1, InGain1, InBuffer2, OutBuffer);
-	}
-
-	/* Performs an element-wise weighted sum OutputBuffer = (InBuffer1 x InGain1) + (InBuffer2 x InGain2) */
-	void BufferWeightedSumFast(const float* RESTRICT InBuffer1, float InGain1, const float* RESTRICT InBuffer2, float InGain2, float* RESTRICT OutBuffer, int32 InNum)
-	{
-		ArrayWeightedSum(MakeArrayView<const float>(InBuffer1, InNum), InGain1, MakeArrayView<const float>(InBuffer2, InNum), InGain2, MakeArrayView<float>(OutBuffer, InNum));
-	}
-
-	/* Performs an element-wise weighted sum OutputBuffer = (InBuffer1 x InGain1) + InBuffer2 */
-	void BufferWeightedSumFast(const float* RESTRICT InBuffer1, float InGain1, const float* RESTRICT InBuffer2, float* RESTRICT OutBuffer, int32 InNum)
-	{
-		ArrayWeightedSum(MakeArrayView<const float>(InBuffer1, InNum), InGain1, MakeArrayView<const float>(InBuffer2, InNum), MakeArrayView<float>(OutBuffer, InNum));
-	}
-
-	void FadeBufferFast(FAlignedFloatBuffer& OutFloatBuffer, const float StartValue, const float EndValue)
-	{
-		ArrayFade(OutFloatBuffer, StartValue, EndValue);
-	}
-
-	void FadeBufferFast(float* RESTRICT OutFloatBuffer, int32 NumSamples, const float StartValue, const float EndValue)
-	{
-		ArrayFade(MakeArrayView<float>(OutFloatBuffer, NumSamples), StartValue, EndValue);
-	}
-
-	void MixInBufferFast(const FAlignedFloatBuffer& InFloatBuffer, FAlignedFloatBuffer& BufferToSumTo, const float Gain)
-	{
-		ArrayMixIn(InFloatBuffer, BufferToSumTo, Gain);
-	}
-
-	void MixInBufferFast(const float* RESTRICT InFloatBuffer, float* RESTRICT BufferToSumTo, int32 NumSamples, const float Gain)
-	{
-		ArrayMixIn(MakeArrayView<const float>(InFloatBuffer, NumSamples), MakeArrayView<float>(BufferToSumTo, NumSamples), Gain);
-	}
-
-	void MixInBufferFast(const FAlignedFloatBuffer& InFloatBuffer, FAlignedFloatBuffer& BufferToSumTo)
-	{
-		ArrayMixIn(InFloatBuffer, BufferToSumTo);
-	}
-
-	void MixInBufferFast(const float* RESTRICT InFloatBuffer, float* RESTRICT BufferToSumTo, int32 NumSamples)
-	{
-		ArrayMixIn(MakeArrayView<const float>(InFloatBuffer, NumSamples), MakeArrayView<float>(BufferToSumTo, NumSamples));
-	}
-
-	void MixInBufferFast(const FAlignedFloatBuffer& InFloatBuffer, FAlignedFloatBuffer& BufferToSumTo, const float StartGain, const float EndGain)
-	{
-		ArrayMixIn(InFloatBuffer, BufferToSumTo, StartGain, EndGain);
-	}
-
-	void MixInBufferFast(const float* RESTRICT InFloatBuffer, float* RESTRICT BufferToSumTo, int32 NumSamples, const float StartGain, const float EndGain)
-	{
-		ArrayMixIn(MakeArrayView<const float>(InFloatBuffer, NumSamples), MakeArrayView<float>(BufferToSumTo, NumSamples), StartGain, EndGain);
-	}
-
-	/* Subtracts two buffers together element-wise. */
-	void BufferSubtractFast(const FAlignedFloatBuffer& InMinuend, const FAlignedFloatBuffer& InSubtrahend, FAlignedFloatBuffer& OutputBuffer)
-	{
-		ArraySubtract(InMinuend, InSubtrahend, OutputBuffer);
-	}
-
-	/* Subtracts two buffers together element-wise. */
-	void BufferSubtractFast(const float* RESTRICT InMinuend, const float* RESTRICT InSubtrahend, float* RESTRICT OutBuffer, int32 InNum)
-	{
-		ArraySubtract(MakeArrayView<const float>(InMinuend, InNum), MakeArrayView<const float>(InSubtrahend, InNum), MakeArrayView<float>(OutBuffer, InNum));
-	}
-
-	/* Performs element-wise in-place subtraction placing the result in the subtrahend. InOutSubtrahend = InMinuend - InOutSubtrahend */
-	void BufferSubtractInPlace1Fast(const FAlignedFloatBuffer& InMinuend, FAlignedFloatBuffer& InOutSubtrahend)
-	{
-		ArraySubtractInPlace1(InMinuend, InOutSubtrahend);
-	}
-
-	/* Performs element-wise in-place subtraction placing the result in the subtrahend. InOutSubtrahend = InMinuend - InOutSubtrahend */
-	void BufferSubtractInPlace1Fast(const float* RESTRICT InMinuend, float* RESTRICT InOutSubtrahend, int32 InNum)
-	{
-		ArraySubtractInPlace1(MakeArrayView<const float>(InMinuend, InNum), MakeArrayView<float>(InOutSubtrahend, InNum));
-	}
-	
-	/* Performs element-wise in-place subtraction placing the result in the minuend. InOutMinuend = InOutMinuend - InSubtrahend */
-	void BufferSubtractInPlace2Fast(FAlignedFloatBuffer& InOutMinuend, const FAlignedFloatBuffer& InSubtrahend)
-	{
-		ArraySubtractInPlace2(InOutMinuend, InSubtrahend);
-	}
-
-	/* Performs element-wise in-place subtraction placing the result in the minuend. InOutMinuend = InOutMinuend - InSubtrahend */
-	void BufferSubtractInPlace2Fast(float* RESTRICT InOutMinuend, const float* RESTRICT InSubtrahend, int32 InNum)
-	{
-		ArraySubtractInPlace2(MakeArrayView<float>(InOutMinuend, InNum), MakeArrayView<const float>(InSubtrahend, InNum));
-	}
-
-	void SumBuffers(const FAlignedFloatBuffer& InFloatBuffer1, const FAlignedFloatBuffer& InFloatBuffer2, FAlignedFloatBuffer& OutputBuffer)
-	{
-		ArraySum(InFloatBuffer1, InFloatBuffer2, OutputBuffer);
-	}
-
-	void SumBuffers(const float* RESTRICT InFloatBuffer1, const float* RESTRICT InFloatBuffer2, float* RESTRICT OutputBuffer, int32 NumSamples)
-	{
-		ArraySum(MakeArrayView<const float>(InFloatBuffer1, NumSamples), MakeArrayView<const float>(InFloatBuffer2, NumSamples), MakeArrayView<float>(OutputBuffer, NumSamples));
-	}
-
-	void MultiplyBuffersInPlace(const FAlignedFloatBuffer& InFloatBuffer, FAlignedFloatBuffer& BufferToMultiply)
-	{
-		ArrayMultiplyInPlace(InFloatBuffer, BufferToMultiply);
-	}
-
-	void MultiplyBuffersInPlace(const float* RESTRICT InFloatBuffer, float* RESTRICT BufferToMultiply, int32 NumSamples)
-	{
-		ArrayMultiplyInPlace(MakeArrayView<const float>(InFloatBuffer, NumSamples), MakeArrayView<float>(BufferToMultiply, NumSamples));
-	}
-
-	float GetMagnitude(const FAlignedFloatBuffer& Buffer)
-	{
-		return ArrayGetMagnitude(Buffer);
-	}
-
-	float GetMagnitude(const float* RESTRICT Buffer, int32 NumSamples)
-	{
-		return ArrayGetMagnitude(MakeArrayView<const float>(Buffer, NumSamples));
-	}
-
-	float BufferGetAverageValue(const FAlignedFloatBuffer& Buffer)
-	{
-		return ArrayGetAverageValue(Buffer);
-	}
-
-	float BufferGetAverageValue(const float* RESTRICT Buffer, int32 NumSamples)
-	{
-		return ArrayGetAverageValue(MakeArrayView<const float>(Buffer, NumSamples));
-	}
-
-	float BufferGetAverageAbsValue(const FAlignedFloatBuffer& Buffer)
-	{
-		return ArrayGetAverageAbsValue(Buffer);
-	}
-
-	float BufferGetAverageAbsValue(const float* RESTRICT Buffer, int32 NumSamples)
-	{
-		return ArrayGetAverageAbsValue(MakeArrayView<const float>(Buffer, NumSamples));
-	}
-
 	/**
 	 * CHANNEL MIXING OPERATIONS:
 	 * To understand these functions, it's best that you have prior experience reading SIMD code.
diff --git a/Engine/Source/Runtime/SignalProcessing/Private/DynamicDelayAPF.cpp b/Engine/Source/Runtime/SignalProcessing/Private/DynamicDelayAPF.cpp
--- a/Engine/Source/Runtime/SignalProcessing/Private/DynamicDelayAPF.cpp
+++ b/Engine/Source/Runtime/SignalProcessing/Private/DynamicDelayAPF.cpp
@@ -123,8 +123,6 @@ namespace Audio
 		// w[n] = x[n] + G * w[n - d]
 		ArraySum(InSamplesView, WorkBufferAView, DelayLineInputView);
 
-		ArrayUnderflowClamp(DelayLineInput);
-
 		// Update delay line
 		IntegerDelayLine->RemoveSamples(InNum);
 		IntegerDelayLine->AddSamples(DelayLineInput.GetData(), InNum);
diff --git a/Engine/Source/Runtime/SignalProcessing/Private/DynamicStateVariableFilter.cpp b/Engine/Source/Runtime/SignalProcessing/Private/DynamicStateVariableFilter.cpp
--- a/Engine/Source/Runtime/SignalProcessing/Private/DynamicStateVariableFilter.cpp
+++ b/Engine/Source/Runtime/SignalProcessing/Private/DynamicStateVariableFilter.cpp
@@ -122,8 +122,6 @@ namespace Audio
 			ProcessHighShelf(InSamples, OutSamples, KeySamples, InNumSamples);
 			break;
 		}
-
-		ClearFilterDenormals();
 	}
 
 	void FDynamicStateVariableFilter::ProcessBell(const float* InSamples, float* OutSamples, const float* KeySamples, const int32 InNumSamples)
@@ -284,20 +282,6 @@ namespace Audio
 		return ConvertToLinear(0.5f * (GainDb + DynamicGain));
 	}
 
-	void FDynamicStateVariableFilter::ClearFilterDenormals()
-	{
-		KeyFilterState.Z1 = Audio::UnderflowClamp(KeyFilterState.Z1);
-		KeyFilterState.Z2 = Audio::UnderflowClamp(KeyFilterState.Z2);
-
-		for (int32 Channel = 0; Channel < NumChannels; ++Channel)
-		{
-			FSVFState& State = FilterState[Channel];
-
-			State.Z1 = Audio::UnderflowClamp(State.Z1);
-			State.Z2 = Audio::UnderflowClamp(State.Z2);
-		}
-	}
-
 	void FDynamicStateVariableFilter::SetFrequency(const float InFrequency)
 	{
 		if (FMath::IsNearlyEqual(Frequency, InFrequency) == false)
diff --git a/Engine/Source/Runtime/SignalProcessing/Private/EnvelopeFollower.cpp b/Engine/Source/Runtime/SignalProcessing/Private/EnvelopeFollower.cpp
--- a/Engine/Source/Runtime/SignalProcessing/Private/EnvelopeFollower.cpp
+++ b/Engine/Source/Runtime/SignalProcessing/Private/EnvelopeFollower.cpp
@@ -21,7 +21,7 @@ namespace Audio
 				Value = (InReleaseSamples * Diff) + InSample;
 			}
 
-			return Audio::UnderflowClamp(Value);
+			return Value;
 		}
 	}
 
@@ -327,7 +327,6 @@ namespace Audio
 			for (int32 SampleIndex = ChannelIndex; SampleIndex < NumSamples; SampleIndex += NumChannels)
 			{
 				Value = Beta * OutBuffer[SampleIndex] + Alpha * Value;
-				Value = Audio::UnderflowClamp(Value);
 				OutBuffer[SampleIndex] = Value;
 			}
 
diff --git a/Engine/Source/Runtime/SignalProcessing/Private/FeedbackDelayNetwork.cpp b/Engine/Source/Runtime/SignalProcessing/Private/FeedbackDelayNetwork.cpp
--- a/Engine/Source/Runtime/SignalProcessing/Private/FeedbackDelayNetwork.cpp
+++ b/Engine/Source/Runtime/SignalProcessing/Private/FeedbackDelayNetwork.cpp
@@ -140,7 +140,7 @@ namespace Audio
 			// Calculate input to delay line. 
 			for (int32 j = 0; j < 4; j++)
 			{
-				DelayInput[j] = UnderflowClamp((ScaledInput + FMO[j]) + (Coefficients.APFG[j] * DelayOut[j]));
+				DelayInput[j] = (ScaledInput + FMO[j]) + (Coefficients.APFG[j] * DelayOut[j]);
 			}
 
 			// Store delay line values
@@ -155,7 +155,7 @@ namespace Audio
 			{
 				AllPassOut[j] = (DelayInput[j] * -Coefficients.APFG[j]) + DelayOut[j];
 				// Note the reasignment of lpz. 
-				LPFZ[j] = UnderflowClamp((AllPassOut[j] * Coefficients.LPFA[j]) + (LPFZ[j] * Coefficients.LPFB[j]));
+				LPFZ[j] = (AllPassOut[j] * Coefficients.LPFA[j]) + (LPFZ[j] * Coefficients.LPFB[j]);
 				OutSampleData[i] += LPFZ[j];
 			}
 
diff --git a/Engine/Source/Runtime/SignalProcessing/Private/InterpolatedOnePole.cpp b/Engine/Source/Runtime/SignalProcessing/Private/InterpolatedOnePole.cpp
--- a/Engine/Source/Runtime/SignalProcessing/Private/InterpolatedOnePole.cpp
+++ b/Engine/Source/Runtime/SignalProcessing/Private/InterpolatedOnePole.cpp
@@ -68,8 +68,7 @@ namespace Audio
 		for (int32 i = 0; i < NumChannels; ++i)
 		{
 			const float InputSample = InputFrame[i];
-			float Yn = InputSample + B1Curr * (Z1Data[i] - InputSample); // LPF
-			Yn = UnderflowClamp(Yn);
+			const float Yn = InputSample + B1Curr * (Z1Data[i] - InputSample); // LPF
 			Z1Data[i] = Yn;
 			OutputFrame[i] = Yn;
 		}
@@ -87,8 +86,7 @@ namespace Audio
 			B1Curr += B1Delta * !ChannelIndex;
 
 			const float InputSample = InputBuffer[SampleIndex];
-			float Yn = InputSample + B1Curr * (Z1Data[ChannelIndex] - InputSample); // LPF
-			Yn = UnderflowClamp(Yn);
+			const float Yn = InputSample + B1Curr * (Z1Data[ChannelIndex] - InputSample); // LPF
 			Z1Data[ChannelIndex] = Yn;
 			OutputBuffer[SampleIndex] = Yn;
 		}
@@ -106,8 +104,7 @@ namespace Audio
 			B1Curr += B1Delta * !ChannelIndex;
 
 			const float InputSample = InOutBuffer[SampleIndex];
-			float Yn = InputSample + B1Curr * (Z1Data[ChannelIndex] - InputSample); // LPF
-			Yn = UnderflowClamp(Yn);
+			const float Yn = InputSample + B1Curr * (Z1Data[ChannelIndex] - InputSample); // LPF
 			Z1Data[ChannelIndex] = Yn;
 			InOutBuffer[SampleIndex] = Yn;
 		}
diff --git a/Engine/Source/Runtime/SignalProcessing/Private/LongDelayAPF.cpp b/Engine/Source/Runtime/SignalProcessing/Private/LongDelayAPF.cpp
--- a/Engine/Source/Runtime/SignalProcessing/Private/LongDelayAPF.cpp
+++ b/Engine/Source/Runtime/SignalProcessing/Private/LongDelayAPF.cpp
@@ -133,19 +133,13 @@ namespace Audio
 
 		VectorRegister4Float VG = MakeVectorRegisterFloat(G, G, G, G);
 		VectorRegister4Float VNG = MakeVectorRegisterFloat(-G, -G, -G, -G);
-		VectorRegister4Float VFMIN = MakeVectorRegisterFloat(FLT_MIN, FLT_MIN, FLT_MIN, FLT_MIN);
-		VectorRegister4Float VNFMIN = MakeVectorRegisterFloat(-FLT_MIN, -FLT_MIN, -FLT_MIN, -FLT_MIN);
 
-		for (int32 i = 0; i < InNum; i += 4)
+		for (int32 i = 0; i < InNum; i += AUDIO_NUM_FLOATS_PER_VECTOR_REGISTER)
 		{
 			VectorRegister4Float VInDelay = VectorLoadAligned(&InDelaySamples[i]);
 			VectorRegister4Float VInSamples = VectorLoadAligned(&InSamples[i]);
 			// w[n] = x[n] + G * w[n - D]
 			VectorRegister4Float VOutDelay = VectorMultiplyAdd(VInDelay, VG, VInSamples);
-			
-			// Underflow clamp
-			VectorRegister4Float Mask = VectorBitwiseAnd(VectorCompareGT(VOutDelay, VNFMIN), VectorCompareLT(VOutDelay, VFMIN));
-			VOutDelay = VectorSelect(Mask, GlobalVectorConstants::FloatZero, VOutDelay);
 			VectorStoreAligned(VOutDelay, &OutDelaySamples[i]);
 
 			// y[n] = -G * w[n] + w[n - D]
@@ -156,7 +150,7 @@ namespace Audio
 		// Calculate allpass for remaining samples that we couldn't SIMD
 		for (int32 i = NumToSIMD; i < InNum; i++)
 		{
-			OutDelaySamples[i] = Audio::UnderflowClamp(InDelaySamples[i] * G + InSamples[i]);
+			OutDelaySamples[i] = InDelaySamples[i] * G + InSamples[i];
 			OutSamples[i] = OutDelaySamples[i] * -G + InDelaySamples[i];
 		}
 	}
diff --git a/Engine/Source/Runtime/SignalProcessing/Private/Reverb.cpp b/Engine/Source/Runtime/SignalProcessing/Private/Reverb.cpp
--- a/Engine/Source/Runtime/SignalProcessing/Private/Reverb.cpp
+++ b/Engine/Source/Runtime/SignalProcessing/Private/Reverb.cpp
@@ -367,7 +367,7 @@ namespace Audio
 			// RIGHT PLATE
 
 			// Get input into right plate by adding the left plate's previous sample (feedback path)
-			float RightPlateInput = APF4Out + UnderflowClamp(LeftPlate.PreviousSample);
+			float RightPlateInput = APF4Out + LeftPlate.PreviousSample;
 
 			// Input -> ModulatedAPF
 			float RightPlateModulatedAPFOut = RightPlate.ModulatedAPF.ProcessAudioSample(RightPlateInput);
@@ -397,7 +397,7 @@ namespace Audio
 			// LEFT PLATE
 
 			// Get input into right plate by adding the left plate's previous sample (feedback path)
-			float LeftPlateInput = APF4Out + UnderflowClamp(RightPlate.PreviousSample);
+			float LeftPlateInput = APF4Out + RightPlate.PreviousSample;
 
 			// Input -> ModulatedAPF
 			float LeftPlateModulatedAPFOut = LeftPlate.ModulatedAPF.ProcessAudioSample(LeftPlateInput);
diff --git a/Engine/Source/Runtime/SignalProcessing/Public/DSP/BiQuadFilter.h b/Engine/Source/Runtime/SignalProcessing/Public/DSP/BiQuadFilter.h
--- a/Engine/Source/Runtime/SignalProcessing/Public/DSP/BiQuadFilter.h
+++ b/Engine/Source/Runtime/SignalProcessing/Public/DSP/BiQuadFilter.h
@@ -31,10 +31,7 @@ namespace Audio
 		FORCEINLINE float ProcessAudio(const float InSample)
 		{
 			// Use the biquad difference eq: y(n) = a0*x(n) + a1*x(n-1) + a2*x(n-2) - b1*y(n-1) - b2*y(n-2) 
-			float Output = A0 * InSample + A1 * X_Z1 + A2 * X_Z2 - B1 * Y_Z1 - B2 * Y_Z2;
-
-			// Clamp the output to 0.0 if in sub-normal float region
-			Output = UnderflowClamp(Output);
+			const float Output = A0 * InSample + A1 * X_Z1 + A2 * X_Z2 - B1 * Y_Z1 - B2 * Y_Z2;
 
 			// Apply the z-transforms
 			Y_Z2 = Y_Z1;
diff --git a/Engine/Source/Runtime/SignalProcessing/Public/DSP/BufferVectorOperations.h b/Engine/Source/Runtime/SignalProcessing/Public/DSP/BufferVectorOperations.h
--- a/Engine/Source/Runtime/SignalProcessing/Public/DSP/BufferVectorOperations.h
+++ b/Engine/Source/Runtime/SignalProcessing/Public/DSP/BufferVectorOperations.h
@@ -8,153 +8,6 @@
 
 namespace Audio
 {
-	/* Sets a values to zero if value is denormal. Denormal numbers significantly slow down floating point operations. */
-	UE_DEPRECATED(5.1, "Use equivalent function in FloatArrayMath.h")
-	SIGNALPROCESSING_API void BufferUnderflowClampFast(FAlignedFloatBuffer& InOutBuffer);
-
-	/* Sets a values to zero if value is denormal. Denormal numbers significantly slow down floating point operations. */
-	UE_DEPRECATED(5.1, "Use equivalent function in FloatArrayMath.h")
-	SIGNALPROCESSING_API void BufferUnderflowClampFast(float* RESTRICT InOutBuffer, const int32 InNum);
-	
-	/* Clamps the values in a buffer between a min and max value.*/
-	UE_DEPRECATED(5.1, "Use equivalent function in FloatArrayMath.h")
-	SIGNALPROCESSING_API void BufferRangeClampFast(FAlignedFloatBuffer& InOutBuffer, float InMinValue, float InMaxValue);
-
-	UE_DEPRECATED(5.1, "Use equivalent function in FloatArrayMath.h")
-	SIGNALPROCESSING_API void BufferRangeClampFast(float* RESTRICT InOutBuffer, const int32 InNum, float InMinValue, float InMaxValue);
-	
-	/** Multiplies the input aligned float buffer with the given value. */
-	UE_DEPRECATED(5.1, "Use equivalent function in FloatArrayMath.h")
-	SIGNALPROCESSING_API void BufferMultiplyByConstant(const FAlignedFloatBuffer& InFloatBuffer, float InValue, FAlignedFloatBuffer& OutFloatBuffer);
-	UE_DEPRECATED(5.1, "Use equivalent function in FloatArrayMath.h")
-	SIGNALPROCESSING_API void BufferMultiplyByConstant(const float* RESTRICT InFloatBuffer, float InValue, float* RESTRICT OutFloatBuffer, const int32 InNumSamples);
-	UE_DEPRECATED(5.1, "Use equivalent function in FloatArrayMath.h")
-	SIGNALPROCESSING_API void BufferMultiplyByConstant(const FAlignedFloatBuffer& InFloatBuffer, float InValue, FAlignedFloatBuffer& OutFloatBuffer);
-
-	/** Similar to BufferMultiplyByConstant, but (a) assumes a buffer length divisible by 4 and (b) performs the multiply in place. */
-	UE_DEPRECATED(5.1, "Use equivalent function in FloatArrayMath.h")
-	SIGNALPROCESSING_API void MultiplyBufferByConstantInPlace(FAlignedFloatBuffer& InBuffer, float InGain);
-	UE_DEPRECATED(5.1, "Use equivalent function in FloatArrayMath.h")
-	SIGNALPROCESSING_API void MultiplyBufferByConstantInPlace(float* RESTRICT InBuffer, int32 NumSamples, float InGain);
-
-	/** Adds a constant to a buffer (useful for DC offset removal) */
-	UE_DEPRECATED(5.1, "Use equivalent function in FloatArrayMath.h")
-	SIGNALPROCESSING_API void AddConstantToBufferInplace(FAlignedFloatBuffer& InBuffer, float Constant);
-	UE_DEPRECATED(5.1, "Use equivalent function in FloatArrayMath.h")
-	SIGNALPROCESSING_API void AddConstantToBufferInplace(float* RESTRICT InBuffer, int32 NumSamples, float Constant);
-
-	/** Sets a constant to a buffer (useful for DC offset application) */
-	UE_DEPRECATED(5.1, "Use equivalent function in FloatArrayMath.h")
-	SIGNALPROCESSING_API void BufferSetToConstantInplace(FAlignedFloatBuffer& InBuffer, float Constant);
-	UE_DEPRECATED(5.1, "Use equivalent function in FloatArrayMath.h")
-	SIGNALPROCESSING_API void BufferSetToConstantInplace(float* RESTRICT InBuffer, int32 NumSamples, float Constant);
-
-	/* Performs an element-wise weighted sum OutputBuffer = (InBuffer1 x InGain1) + (InBuffer2 x InGain2) */
-	UE_DEPRECATED(5.1, "Use equivalent function in FloatArrayMath.h")
-	SIGNALPROCESSING_API void BufferWeightedSumFast(const FAlignedFloatBuffer& InBuffer1, float InGain1, const FAlignedFloatBuffer& InBuffer2, float InGain2, FAlignedFloatBuffer& OutBuffer);
-
-	/* Performs an element-wise weighted sum OutputBuffer = (InBuffer1 x InGain1) + InBuffer2 */
-	UE_DEPRECATED(5.1, "Use equivalent function in FloatArrayMath.h")
-	SIGNALPROCESSING_API void BufferWeightedSumFast(const FAlignedFloatBuffer& InBuffer1, float InGain1, const FAlignedFloatBuffer& InBuffer2, FAlignedFloatBuffer& OutBuffer);
-
-	/* Performs an element-wise weighted sum OutputBuffer = (InBuffer1 x InGain1) + (InBuffer2 x InGain2) */
-	UE_DEPRECATED(5.1, "Use equivalent function in FloatArrayMath.h")
-	SIGNALPROCESSING_API void BufferWeightedSumFast(const float* RESTRICT InBuffer1, float InGain1, const float* RESTRICT InBuffer2, float InGain2, float* RESTRICT OutBuffer, int32 InNum);
-
-	/* Performs an element-wise weighted sum OutputBuffer = (InBuffer1 x InGain1) + InBuffer2 */
-	UE_DEPRECATED(5.1, "Use equivalent function in FloatArrayMath.h")
-	SIGNALPROCESSING_API void BufferWeightedSumFast(const float* RESTRICT InBuffer1, float InGain1, const float* RESTRICT InBuffer2, float* RESTRICT OutBuffer, int32 InNum);
-	UE_DEPRECATED(5.1, "Use equivalent function in FloatArrayMath.h")
-	SIGNALPROCESSING_API void MultiplyBufferByConstantInPlace(FAlignedFloatBuffer& InBuffer, float InGain);
-	UE_DEPRECATED(5.1, "Use equivalent function in FloatArrayMath.h")
-	SIGNALPROCESSING_API void MultiplyBufferByConstantInPlace(float* RESTRICT InBuffer, int32 NumSamples, float InGain);
-
-	/* Takes a float buffer and quickly interpolates it's gain from StartValue to EndValue. */
-	/* This operation completely ignores channel counts, so avoid using this function on buffers that are not mono, stereo or quad */
-	/* if the buffer needs to fade all channels uniformly. */
-	UE_DEPRECATED(5.1, "Use equivalent function in FloatArrayMath.h")
-	SIGNALPROCESSING_API void FadeBufferFast(FAlignedFloatBuffer& OutFloatBuffer, const float StartValue, const float EndValue);
-	UE_DEPRECATED(5.1, "Use equivalent function in FloatArrayMath.h")
-	SIGNALPROCESSING_API void FadeBufferFast(float* RESTRICT OutFloatBuffer, int32 NumSamples, const float StartValue, const float EndValue);
-
-	/** Takes buffer InFloatBuffer, optionally multiplies it by Gain, and adds it to BufferToSumTo. */
-	UE_DEPRECATED(5.1, "Use equivalent function in FloatArrayMath.h")
-	SIGNALPROCESSING_API void MixInBufferFast(const FAlignedFloatBuffer& InFloatBuffer, FAlignedFloatBuffer& BufferToSumTo, const float Gain);
-	UE_DEPRECATED(5.1, "Use equivalent function in FloatArrayMath.h")
-	SIGNALPROCESSING_API void MixInBufferFast(const float* RESTRICT InFloatBuffer, float* RESTRICT BufferToSumTo, int32 NumSamples, const float Gain);
-	UE_DEPRECATED(5.1, "Use equivalent function in FloatArrayMath.h")
-	SIGNALPROCESSING_API void MixInBufferFast(const FAlignedFloatBuffer& InFloatBuffer, FAlignedFloatBuffer& BufferToSumTo, const float StartGain, const float EndGain);
-	UE_DEPRECATED(5.1, "Use equivalent function in FloatArrayMath.h")
-	SIGNALPROCESSING_API void MixInBufferFast(const float* RESTRICT InFloatBuffer, float* RESTRICT BufferToSumTo, int32 NumSamples, const float StartGain, const float EndGain);
-	UE_DEPRECATED(5.1, "Use equivalent function in FloatArrayMath.h")
-	SIGNALPROCESSING_API void MixInBufferFast(const FAlignedFloatBuffer& InFloatBuffer, FAlignedFloatBuffer& BufferToSumTo);
-	UE_DEPRECATED(5.1, "Use equivalent function in FloatArrayMath.h")
-	SIGNALPROCESSING_API void MixInBufferFast(const float* RESTRICT InFloatBuffer, float* RESTRICT BufferToSumTo, int32 NumSamples);
-	UE_DEPRECATED(5.1, "Use equivalent function in FloatArrayMath.h")
-
-	/* Subtracts two buffers together element-wise. */
-	UE_DEPRECATED(5.1, "Use equivalent function in FloatArrayMath.h")
-	SIGNALPROCESSING_API void BufferSubtractFast(const FAlignedFloatBuffer& InMinuend, const FAlignedFloatBuffer& InSubtrahend, FAlignedFloatBuffer& OutputBuffer);
-
-	/* Subtracts two buffers together element-wise. */
-	UE_DEPRECATED(5.1, "Use equivalent function in FloatArrayMath.h")
-	SIGNALPROCESSING_API void BufferSubtractFast(const float* RESTRICT InMinuend, const float* RESTRICT InSubtrahend, float* RESTRICT OutputBuffer, int32 NumSamples);
-
-	/* Performs element-wise in-place subtraction placing the result in the subtrahend. InOutSubtrahend = InMinuend - InOutSubtrahend */
-	UE_DEPRECATED(5.1, "Use equivalent function in FloatArrayMath.h")
-	SIGNALPROCESSING_API void BufferSubtractInPlace1Fast(const FAlignedFloatBuffer& InMinuend, FAlignedFloatBuffer& InOutSubtrahend);
-
-	/* Performs element-wise in-place subtraction placing the result in the subtrahend. InOutSubtrahend = InMinuend - InOutSubtrahend */
-	UE_DEPRECATED(5.1, "Use equivalent function in FloatArrayMath.h")
-	SIGNALPROCESSING_API void BufferSubtractInPlace1Fast(const float* RESTRICT InMinuend, float* RESTRICT InOutSubtrahend, int32 NumSamples);
-
-
-	/* Performs element-wise in-place subtraction placing the result in the minuend. InOutMinuend = InOutMinuend - InSubtrahend */
-	UE_DEPRECATED(5.1, "Use equivalent function in FloatArrayMath.h")
-	SIGNALPROCESSING_API void BufferSubtractInPlace2Fast(FAlignedFloatBuffer& InOutMinuend, const FAlignedFloatBuffer& InSubtrahend);
-
-	/* Performs element-wise in-place subtraction placing the result in the minuend. InOutMinuend = InOutMinuend - InSubtrahend */
-	UE_DEPRECATED(5.1, "Use equivalent function in FloatArrayMath.h")
-	SIGNALPROCESSING_API void BufferSubtractInPlace2Fast(float* RESTRICT InOutMinuend, const float* RESTRICT InSubtrahend, int32 NumSamples);
-
-	/** This version of MixInBufferFast will fade from StartGain to EndGain. */
-	UE_DEPRECATED(5.1, "Use equivalent function in FloatArrayMath.h")
-	SIGNALPROCESSING_API void MixInBufferFast(const FAlignedFloatBuffer& InFloatBuffer, FAlignedFloatBuffer& BufferToSumTo, const float StartGain, const float EndGain);
-	UE_DEPRECATED(5.1, "Use equivalent function in FloatArrayMath.h")
-	SIGNALPROCESSING_API void MixInBufferFast(const float* RESTRICT InFloatBuffer, float* RESTRICT BufferToSumTo, int32 NumSamples, const float StartGain, const float EndGain);
-
-	/** Sums two buffers together and places the result in the resulting buffer. */
-	UE_DEPRECATED(5.1, "Use equivalent function in FloatArrayMath.h")
-	SIGNALPROCESSING_API void SumBuffers(const FAlignedFloatBuffer& InFloatBuffer1, const FAlignedFloatBuffer& InFloatBuffer2, FAlignedFloatBuffer& OutputBuffer);
-	UE_DEPRECATED(5.1, "Use equivalent function in FloatArrayMath.h")
-	SIGNALPROCESSING_API void SumBuffers(const float* RESTRICT InFloatBuffer1, const float* RESTRICT InFloatBuffer2, float* RESTRICT OutputBuffer, int32 NumSamples);
-
-	/** Multiply the second buffer in place by the first buffer. */
-	UE_DEPRECATED(5.1, "Use equivalent function in FloatArrayMath.h")
-	SIGNALPROCESSING_API void MultiplyBuffersInPlace(const FAlignedFloatBuffer& InFloatBuffer, FAlignedFloatBuffer& BufferToMultiply);
-	UE_DEPRECATED(5.1, "Use equivalent function in FloatArrayMath.h")
-	SIGNALPROCESSING_API void MultiplyBuffersInPlace(const float* RESTRICT InFloatBuffer, float* RESTRICT BufferToMultiply, int32 NumSamples);
-
-	/** CHANNEL-AGNOSTIC ANALYSIS OPERATIONS */
-
-	/** Takes an audio buffer and returns the magnitude across that buffer. */
-	UE_DEPRECATED(5.1, "Use equivalent function in FloatArrayMath.h")
-	SIGNALPROCESSING_API float GetMagnitude(const FAlignedFloatBuffer& Buffer);
-	UE_DEPRECATED(5.1, "Use equivalent function in FloatArrayMath.h")
-	SIGNALPROCESSING_API float GetMagnitude(const float* RESTRICT Buffer, int32 NumSamples);
-
-	/** Takes an audio buffer and gets the average amplitude across that buffer. */
-	UE_DEPRECATED(5.1, "Use equivalent function in FloatArrayMath.h")
-	SIGNALPROCESSING_API float BufferGetAverageValue(const FAlignedFloatBuffer& Buffer);
-	UE_DEPRECATED(5.1, "Use equivalent function in FloatArrayMath.h")
-	SIGNALPROCESSING_API float BufferGetAverageValue(const float* RESTRICT Buffer, int32 NumSamples);
-
-	/** Takes an audio buffer and gets the average absolute amplitude across that buffer. */
-	UE_DEPRECATED(5.1, "Use equivalent function in FloatArrayMath.h")
-	SIGNALPROCESSING_API float BufferGetAverageAbsValue(const FAlignedFloatBuffer& Buffer);
-	UE_DEPRECATED(5.1, "Use equivalent function in FloatArrayMath.h")
-	SIGNALPROCESSING_API float BufferGetAverageAbsValue(const float* RESTRICT Buffer, int32 NumSamples);
-
 	/** CHANNEL-SPECIFIC OPERATIONS */
 
 	/** Takes a 2 channel interleaved buffer and applies Gains to it. Gains is expected to point to a 2 float long buffer.
diff --git a/Engine/Source/Runtime/SignalProcessing/Public/DSP/Dsp.h b/Engine/Source/Runtime/SignalProcessing/Public/DSP/Dsp.h
--- a/Engine/Source/Runtime/SignalProcessing/Public/DSP/Dsp.h
+++ b/Engine/Source/Runtime/SignalProcessing/Public/DSP/Dsp.h
@@ -36,6 +36,7 @@ namespace Audio
 	}
 
 	// Clamps floats to 0 if they are in sub-normal range
+	UE_DEPRECATED(5.5, "Audio code relies on denormals flush to zero floating point mode from now on. Please use FScopedFTZFloatMode instead of this API")
 	static FORCEINLINE float UnderflowClamp(const float InValue)
 	{
 		if (InValue > -FLT_MIN && InValue < FLT_MIN)
diff --git a/Engine/Source/Runtime/SignalProcessing/Public/DSP/DynamicStateVariableFilter.h b/Engine/Source/Runtime/SignalProcessing/Public/DSP/DynamicStateVariableFilter.h
--- a/Engine/Source/Runtime/SignalProcessing/Public/DSP/DynamicStateVariableFilter.h
+++ b/Engine/Source/Runtime/SignalProcessing/Public/DSP/DynamicStateVariableFilter.h
@@ -125,7 +125,5 @@ namespace Audio
 		SIGNALPROCESSING_API void UpdatePostFilterVariables();
 		
 		SIGNALPROCESSING_API float CalcGain(const float KeySample);
-		
-		SIGNALPROCESSING_API void ClearFilterDenormals();
 	};
 }
diff --git a/Engine/Source/Runtime/SignalProcessing/Public/DSP/EnvelopeFollower.h b/Engine/Source/Runtime/SignalProcessing/Public/DSP/EnvelopeFollower.h
--- a/Engine/Source/Runtime/SignalProcessing/Public/DSP/EnvelopeFollower.h
+++ b/Engine/Source/Runtime/SignalProcessing/Public/DSP/EnvelopeFollower.h
@@ -382,7 +382,6 @@ namespace Audio
 				Value = (GetReleaseTimeSamples() * Diff) + NormSample;
 			}
 
-			Value = Audio::UnderflowClamp(Value);
 			return Value;
 		}
 
@@ -404,7 +403,6 @@ namespace Audio
 				{
 					Value = (GetReleaseTimeSamples() * Diff) + NormSample;
 				}
-				Value = Audio::UnderflowClamp(Value);
 				OutSamples[SampleIndex] = Value;
 			}
 		}
@@ -433,7 +431,6 @@ namespace Audio
 					{
 						float SampleSquared = InSample * InSample;
 						AnalysisValue = AnalysisFilterBeta * SampleSquared + AnalysisFilterAlpha * AnalysisValue;
-						AnalysisValue = Audio::UnderflowClamp(AnalysisValue);
 						return AnalysisValue;
 					}
 
@@ -442,7 +439,6 @@ namespace Audio
 					{
 						float SampleSquared = InSample * InSample;
 						AnalysisValue = AnalysisFilterBeta * SampleSquared + AnalysisFilterAlpha * AnalysisValue;
-						AnalysisValue = Audio::UnderflowClamp(AnalysisValue);
 						return FMath::Sqrt(AnalysisValue);
 					}
 
@@ -471,7 +467,6 @@ namespace Audio
 				for (int32 SampleIndex = 0; SampleIndex < InNumSamples; ++SampleIndex)
 				{
 					AnalysisValue = AnalysisFilterBeta * OutSamples[SampleIndex] + AnalysisFilterAlpha * AnalysisValue;
-					AnalysisValue = Audio::UnderflowClamp(AnalysisValue);
 					OutSamples[SampleIndex] = AnalysisValue;
 				}
 				return;
@@ -484,7 +479,6 @@ namespace Audio
 				for (int32 SampleIndex = 0; SampleIndex < InNumSamples; ++SampleIndex)
 				{
 					AnalysisValue = AnalysisFilterBeta * OutSamples[SampleIndex] + AnalysisFilterAlpha * AnalysisValue;
-					AnalysisValue = Audio::UnderflowClamp(AnalysisValue);
 					OutSamples[SampleIndex] = AnalysisValue;
 				}
 				ArraySqrtInPlace(Out);
diff --git a/Engine/Source/Runtime/SignalProcessing/Public/DSP/FloatArrayMath.h b/Engine/Source/Runtime/SignalProcessing/Public/DSP/FloatArrayMath.h
--- a/Engine/Source/Runtime/SignalProcessing/Public/DSP/FloatArrayMath.h
+++ b/Engine/Source/Runtime/SignalProcessing/Public/DSP/FloatArrayMath.h
@@ -206,6 +206,7 @@ namespace Audio
 	SIGNALPROCESSING_API void ArrayComplexToPower(TArrayView<const float> InRealSamples, TArrayView<const float> InImaginarySamples, TArrayView<float> OutPowerSamples);
 
 	/* Sets a values to zero if value is denormal. Denormal numbers significantly slow down floating point operations. */
+	UE_DEPRECATED(5.5, "Audio code relies on denormals flush to zero floating point mode from now on. Please use FScopedFTZFloatMode instead of this API")
 	SIGNALPROCESSING_API void ArrayUnderflowClamp(TArrayView<float> InOutBuffer);
 
 	/* Clamps the values in a buffer between a min and max value.*/
diff --git a/Engine/Source/Runtime/SignalProcessing/Public/DSP/OnePole.h b/Engine/Source/Runtime/SignalProcessing/Public/DSP/OnePole.h
--- a/Engine/Source/Runtime/SignalProcessing/Public/DSP/OnePole.h
+++ b/Engine/Source/Runtime/SignalProcessing/Public/DSP/OnePole.h
@@ -48,8 +48,7 @@ namespace Audio
 
 		FORCEINLINE float ProcessAudioSample(const float InputSample)
 		{
-			float Yn = InputSample*A0 + B1*Z1;
-			Yn = UnderflowClamp(Yn);
+			const float Yn = InputSample*A0 + B1*Z1;
 			Z1 = Yn;
 			return Yn;
 		}
@@ -127,8 +126,7 @@ namespace Audio
 		{			
 			for (int32 i = 0; i < NumChannels; ++i)
 			{
-				float Yn = InputFrame[i] * A0 + B1*DelayPtr[i];
-				Yn = UnderflowClamp(Yn);
+				const float Yn = InputFrame[i] * A0 + B1*DelayPtr[i];
 				DelayPtr[i] = Yn;
 				OutputFrame[i] = Yn;
 			}
