Optimized Audio::FBiquadFilter

diff --git a/Engine/Source/Runtime/SignalProcessing/Private/Filter.cpp b/Engine/Source/Runtime/SignalProcessing/Private/Filter.cpp
--- a/Engine/Source/Runtime/SignalProcessing/Private/Filter.cpp
+++ b/Engine/Source/Runtime/SignalProcessing/Private/Filter.cpp
@@ -5,6 +5,60 @@
 
 namespace Audio
 {
+	// Simple biquad filter structure handling a biquad formulation
+	// See: https://en.wikipedia.org/wiki/Digital_biquad_filter
+	// Calculations of coefficients are handled outside this class.
+	// Filter coefficients are public and are intended to be used externally.
+	struct FBiquadFilter::FBiquadCoeff
+	{
+	public:
+		FBiquadCoeff()
+			: A0(1.0f)
+			, A1(0.0f)
+			, A2(0.0f)
+			, B1(0.0f)
+			, B2(0.0f)
+		{
+			Reset();
+		}
+
+		// Reset the filter (flush delays)
+		void Reset()
+		{
+			X_Z1 = 0.0f;
+			X_Z2 = 0.0f;
+			Y_Z1 = 0.0f;
+			Y_Z2 = 0.0f;
+		}
+
+		FORCEINLINE float ProcessAudio(const float InSample)
+		{
+			// Use the biquad difference eq: y(n) = a0*x(n) + a1*x(n-1) + a2*x(n-2) - b1*y(n-1) - b2*y(n-2) 
+			const float Output = A0 * InSample + A1 * X_Z1 + A2 * X_Z2 - B1 * Y_Z1 - B2 * Y_Z2;
+
+			// Apply the z-transforms
+			Y_Z2 = Y_Z1;
+			Y_Z1 = Output;
+
+			X_Z2 = X_Z1;
+			X_Z1 = InSample;
+
+			return Output;
+		}
+
+		// Biquad filter coefficients
+		float A0;
+		float A1;
+		float A2;
+		float B1;
+		float B2;
+
+		float X_Z1; // previous inputs z-transforms
+		float X_Z2;
+		float Y_Z1; // prvious outputs z-transforms
+		float Y_Z2;
+	};
+
     float FBiquadFilter::ClampCutoffFrequency(float InCutoffFrequency)
     {
         return FMath::Clamp(InCutoffFrequency, 5.0f, 0.9f * (SampleRate / 2.0f));
@@ -44,7 +98,7 @@ namespace Audio
 			delete[] Biquad;
 		}
 
-		Biquad = new FBiquad[NumChannels];
+		Biquad = new FBiquadCoeff[NumChannels];
 		Reset();
 		CalculateBiquadCoefficients();
 	}
@@ -68,10 +122,37 @@ namespace Audio
 		{
 			if (NumChannels == 1)
 			{
+				// keep these in registers to prevent loading and storing them on every loop iteration
+				float A0 = Biquad->A0;
+				float A1 = Biquad->A1;
+				float A2 = Biquad->A2;
+				float B1 = Biquad->B1;
+				float B2 = Biquad->B2;
+				float X_Z1 = Biquad->X_Z1;
+				float X_Z2 = Biquad->X_Z2;
+				float Y_Z1 = Biquad->Y_Z1;
+				float Y_Z2 = Biquad->Y_Z2;
+
 				for (int32 SampleIndex = 0; SampleIndex < InNumSamples; ++SampleIndex)
 				{
-					OutBuffer[SampleIndex] = Biquad->ProcessAudio(InBuffer[SampleIndex]);
+					const float InSample = InBuffer[SampleIndex];
+					// Use the biquad difference eq: y(n) = a0*x(n) + a1*x(n-1) + a2*x(n-2) - b1*y(n-1) - b2*y(n-2) 
+					const float Output = A0 * InSample + A1 * X_Z1 + A2 * X_Z2 - B1 * Y_Z1 - B2 * Y_Z2;
+
+					// Apply the z-transforms
+					Y_Z2 = Y_Z1;
+					Y_Z1 = Output;
+
+					X_Z2 = X_Z1;
+					X_Z1 = InSample;
+
+					OutBuffer[SampleIndex] = Output;
 				}
+
+				Biquad->X_Z1 = X_Z1;
+				Biquad->X_Z2 = X_Z2;
+				Biquad->Y_Z1 = Y_Z1;
+				Biquad->Y_Z2 = Y_Z2;
 			}
 			else
 			{
diff --git a/Engine/Source/Runtime/SignalProcessing/Public/DSP/BiQuadFilter.h b/Engine/Source/Runtime/SignalProcessing/Public/DSP/BiQuadFilter.h
--- a/Engine/Source/Runtime/SignalProcessing/Public/DSP/BiQuadFilter.h
+++ b/Engine/Source/Runtime/SignalProcessing/Public/DSP/BiQuadFilter.h
@@ -11,6 +11,7 @@ namespace Audio
 	// See: https://en.wikipedia.org/wiki/Digital_biquad_filter
 	// Calculations of coefficients are handled outside this class.
 	// Filter coefficients are public and are intended to be used externally.
+	UE_DEPRECATED(5.5, "Audio::FBiquad is deprecated in favor of Audio::FBiquadFilter due to performance issues")
 	class FBiquad
 	{
 	public:
@@ -24,10 +25,6 @@ namespace Audio
 			Reset();
 		}
 
-		virtual ~FBiquad()
-		{
-		}
-
 		FORCEINLINE float ProcessAudio(const float InSample)
 		{
 			// Use the biquad difference eq: y(n) = a0*x(n) + a1*x(n-1) + a2*x(n-2) - b1*y(n-1) - b2*y(n-2) 
diff --git a/Engine/Source/Runtime/SignalProcessing/Public/DSP/Filter.h b/Engine/Source/Runtime/SignalProcessing/Public/DSP/Filter.h
--- a/Engine/Source/Runtime/SignalProcessing/Public/DSP/Filter.h
+++ b/Engine/Source/Runtime/SignalProcessing/Public/DSP/Filter.h
@@ -4,7 +4,7 @@
 
 #include "CoreMinimal.h"
 #include "DSP/ModulationMatrix.h"
-#include "DSP/BiQuadFilter.h"
+#include "DSP/Dsp.h"
 
 
 namespace Audio
@@ -74,6 +74,7 @@ namespace Audio
 		SIGNALPROCESSING_API void SetEnabled(const bool bInEnabled);
 
 	protected:
+		struct FBiquadCoeff;
 
 		// Function computes biquad coefficients based on current filter settings
 		SIGNALPROCESSING_API void CalculateBiquadCoefficients();
@@ -85,7 +86,7 @@ namespace Audio
 		EBiquadFilter::Type FilterType;
 
 		// Biquad filter objects for each channel
-		FBiquad* Biquad;
+		FBiquadCoeff* Biquad;
 
 		// The sample rate of the filter
 		float SampleRate;
