Vectorized FInterpolatedLPF\FInterpolatedHPF

diff --git a/Engine/Source/Runtime/SignalProcessing/Private/InterpolatedOnePole.cpp b/Engine/Source/Runtime/SignalProcessing/Private/InterpolatedOnePole.cpp
--- a/Engine/Source/Runtime/SignalProcessing/Private/InterpolatedOnePole.cpp
+++ b/Engine/Source/Runtime/SignalProcessing/Private/InterpolatedOnePole.cpp
@@ -7,6 +7,7 @@
 
 namespace Audio
 {
+	static constexpr int32 Simd4Mask = 0xFFFFFFFC;
 
 	// INTERPOLATED ONE-POLE LOW-PASS IMPLEMENTATION
 	FInterpolatedLPF::FInterpolatedLPF()
@@ -76,7 +77,87 @@ namespace Audio
 
 	void FInterpolatedLPF::ProcessAudioBuffer(const float *RESTRICT InputBuffer, float *RESTRICT OutputBuffer, const int32 NumSamples)
 	{
-		for (int SampleIndex = 0; SampleIndex < NumSamples; ++SampleIndex)
+		if (NumChannels == 1)
+		{
+			// keep these in registers instead of having a load\store on every loop iteration
+			float B1 = B1Curr;
+			float Z1D = Z1Data[0];
+			float Delta = B1Delta;
+			for (int32 SampleIndex = 0; SampleIndex < NumSamples; ++SampleIndex)
+			{
+				// step forward coefficient
+				B1 += Delta;
+
+				const float InputSample = InputBuffer[SampleIndex];
+				float Yn = InputSample + B1Curr * (Z1D - InputSample); // LPF
+				Z1D = Yn;
+				OutputBuffer[SampleIndex] = Yn;
+			}
+
+			Z1Data[0] = Z1D;
+			B1Curr = B1;
+
+			return;
+		}
+
+		int32 SampleIndex = 0;
+		const int32 NumToSimd = NumSamples & Simd4Mask;
+		if (NumChannels == 2)
+		{
+			VectorRegister4Float B1CurrVec = VectorLoadFloat1(&B1Curr);
+			VectorRegister4Float B1DeltaVec = MakeVectorRegister(B1Delta + B1Delta, B1Delta + B1Delta, B1Delta + B1Delta, B1Delta + B1Delta);
+			VectorRegister4Float Z1DataVec = MakeVectorRegister(Z1Data[0], Z1Data[1], Z1Data[0], Z1Data[1]);
+			
+			if (NumToSimd)
+			{
+				B1CurrVec = VectorAdd(B1CurrVec, MakeVectorRegister(B1Delta, B1Delta, B1Delta + B1Delta, B1Delta + B1Delta));	//b1 + delta, b1 + delta, b1 + 2x delta, b1 + 2x delta
+			}
+
+			for (; SampleIndex < NumToSimd; SampleIndex += AUDIO_NUM_FLOATS_PER_VECTOR_REGISTER)
+			{
+				VectorRegister4Float InputSample = VectorLoad(&InputBuffer[SampleIndex]);
+				VectorRegister4Float Yn = VectorMultiplyAdd(B1CurrVec, VectorSubtract(Z1DataVec, InputSample), InputSample);	//Y0, Y1, garbage, garbage
+
+				Z1DataVec = VectorShuffle(Z1DataVec, Yn, 0, 1, 0, 1);	//Z0, Z1, Y0, Y1
+				Yn = VectorMultiplyAdd(B1CurrVec, VectorSubtract(Z1DataVec, InputSample), InputSample);	//Y0, Y1, Y2, Y3
+
+				VectorStore(Yn, &OutputBuffer[SampleIndex]);
+				Z1DataVec = VectorSwizzle(Yn, 2, 3, 2, 3);	//Y2, Y3, Y2, Y3
+
+				B1CurrVec = VectorAdd(B1CurrVec, B1DeltaVec);	//b1 + 2x delta, b1 + 2x delta, b1 + 2x delta, b1 + 2x delta
+			}
+
+			float Store[4];
+			VectorStore(B1CurrVec, Store);
+			B1Curr = Store[0];
+			VectorStore(Z1DataVec, Store);
+			Z1Data[0] = Store[0];
+			Z1Data[1] = Store[1];
+		}
+		else if (NumChannels == 4)
+		{
+			VectorRegister4Float B1CurrVec = VectorLoadFloat1(&B1Curr);
+			VectorRegister4Float B1DeltaVec = VectorLoadFloat1(&B1Delta);
+			VectorRegister4Float Z1DataVec = VectorLoad(&Z1Data[0]);
+
+			for (; SampleIndex < NumToSimd; SampleIndex += AUDIO_NUM_FLOATS_PER_VECTOR_REGISTER)
+			{
+				B1CurrVec = VectorAdd(B1CurrVec, B1DeltaVec);
+
+				VectorRegister4Float InputSample = VectorLoad(&InputBuffer[SampleIndex]);
+				VectorRegister4Float Yn = VectorMultiplyAdd(B1CurrVec, VectorSubtract(Z1DataVec, InputSample), InputSample);
+
+				VectorStore(Yn, &OutputBuffer[SampleIndex]);
+				Z1DataVec = Yn;
+			}
+
+			float Store[4];
+			VectorStore(B1CurrVec, Store);
+			B1Curr = Store[0];
+			VectorStore(Z1DataVec, &Z1Data[0]);
+		}
+
+		for (; SampleIndex < NumSamples; ++SampleIndex)
 		{
 			// cache which delay term we should be using
 			const int32 ChannelIndex = SampleIndex % NumChannels;
@@ -94,7 +175,87 @@ namespace Audio
 
 	void FInterpolatedLPF::ProcessBufferInPlace(float* InOutBuffer, const int32 NumSamples)
 	{
-		for (int32 SampleIndex = 0; SampleIndex < NumSamples; ++SampleIndex)
+		if (NumChannels == 1)
+		{
+			// keep these in registers instead of having a load\store on every loop iteration
+			float B1 = B1Curr;
+			float Z1D = Z1Data[0];
+			float Delta = B1Delta;
+			for (int32 SampleIndex = 0; SampleIndex < NumSamples; ++SampleIndex)
+			{
+				// step forward coefficient
+				B1 += Delta;
+
+				const float InputSample = InOutBuffer[SampleIndex];
+				float Yn = InputSample + B1Curr * (Z1D - InputSample); // LPF
+				Z1D = Yn;
+				InOutBuffer[SampleIndex] = Yn;
+			}
+
+			Z1Data[0] = Z1D;
+			B1Curr = B1;
+
+			return;
+		}
+
+		int32 SampleIndex = 0;
+		const int32 NumToSimd = NumSamples & Simd4Mask;
+		if (NumChannels == 2)
+		{
+			VectorRegister4Float B1CurrVec = VectorLoadFloat1(&B1Curr);
+			VectorRegister4Float B1DeltaVec = MakeVectorRegister(B1Delta + B1Delta, B1Delta + B1Delta, B1Delta + B1Delta, B1Delta + B1Delta);
+			VectorRegister4Float Z1DataVec = MakeVectorRegister(Z1Data[0], Z1Data[1], Z1Data[0], Z1Data[1]);
+
+			if (NumToSimd)
+			{
+				B1CurrVec = VectorAdd(B1CurrVec, MakeVectorRegister(B1Delta, B1Delta, B1Delta + B1Delta, B1Delta + B1Delta));	//b1 + delta, b1 + delta, b1 + 2x delta, b1 + 2x delta
+			}
+
+			for (; SampleIndex < NumToSimd; SampleIndex += 4)
+			{
+				VectorRegister4Float InputSample = VectorLoad(&InOutBuffer[SampleIndex]);
+				VectorRegister4Float Yn = VectorMultiplyAdd(B1CurrVec, VectorSubtract(Z1DataVec, InputSample), InputSample);	//Y0, Y1, garbage, garbage
+
+				Z1DataVec = VectorShuffle(Z1DataVec, Yn, 0, 1, 0, 1);	//Z0, Z1, Y0, Y1
+				Yn = VectorMultiplyAdd(B1CurrVec, VectorSubtract(Z1DataVec, InputSample), InputSample);	//Y0, Y1, Y2, Y3
+
+				VectorStore(Yn, &InOutBuffer[SampleIndex]);
+				Z1DataVec = VectorSwizzle(Yn, 2, 3, 2, 3);	//Y2, Y3, Y2, Y3
+
+				B1CurrVec = VectorAdd(B1CurrVec, B1DeltaVec);	//b1 + 2x delta, b1 + 2x delta, b1 + 2x delta, b1 + 2x delta
+			}
+
+			float Store[4];
+			VectorStore(B1CurrVec, Store);
+			B1Curr = Store[0];
+			VectorStore(Z1DataVec, Store);
+			Z1Data[0] = Store[0];
+			Z1Data[1] = Store[1];
+		}
+		else if (NumChannels == 4)
+		{
+			VectorRegister4Float B1CurrVec = VectorLoadFloat1(&B1Curr);
+			VectorRegister4Float B1DeltaVec = VectorLoadFloat1(&B1Delta);
+			VectorRegister4Float Z1DataVec = VectorLoad(&Z1Data[0]);
+
+			for (; SampleIndex < NumToSimd; SampleIndex += 4)
+			{
+				B1CurrVec = VectorAdd(B1CurrVec, B1DeltaVec);
+
+				VectorRegister4Float InputSample = VectorLoad(&InOutBuffer[SampleIndex]);
+				VectorRegister4Float Yn = VectorMultiplyAdd(B1CurrVec, VectorSubtract(Z1DataVec, InputSample), InputSample);
+
+				VectorStore(Yn, &InOutBuffer[SampleIndex]);
+				Z1DataVec = Yn;
+			}
+
+			float Store[4];
+			VectorStore(B1CurrVec, Store);
+			B1Curr = Store[0];
+			VectorStore(Z1DataVec, &Z1Data[0]);
+		}
+
+		for (; SampleIndex < NumSamples; ++SampleIndex)
 		{
 			// cache which delay term we should be using
 			const int32 ChannelIndex = SampleIndex % NumChannels;
@@ -192,7 +353,89 @@ namespace Audio
 
 	void FInterpolatedHPF::ProcessAudioBuffer(const float *RESTRICT InputBuffer, float *RESTRICT OutputBuffer, const int32 NumSamples)
 	{
-		for (int SampleIndex = 0; SampleIndex < NumSamples; ++SampleIndex)
+		if (NumChannels == 1)
+		{
+			// keep these in registers instead of having a load\store on every loop iteration
+			float A0 = A0Curr;
+			float Z1D = Z1Data[0];
+			float Delta = A0Delta;
+			for (int32 SampleIndex = 0; SampleIndex < NumSamples; ++SampleIndex)
+			{
+				// step forward coefficient
+				A0 += Delta;
+
+				const float InputSample = InputBuffer[SampleIndex];
+				const float Vn = (InputSample - Z1D) * A0;
+				const float LPF = Vn + Z1D;
+				Z1D = Vn + LPF;
+
+				OutputBuffer[SampleIndex] = InputSample - LPF;
+			}
+
+			Z1Data[0] = Z1D;
+			A0Curr = A0;
+
+			return;
+		}
+
+		int32 SampleIndex = 0;
+		const int32 NumToSimd = NumSamples & Simd4Mask;
+		if (NumChannels == 2)
+		{
+			VectorRegister4Float A0CurrVec = VectorLoadFloat1(&A0Curr);
+			VectorRegister4Float A0DeltaVec = MakeVectorRegister(A0Delta + A0Delta, A0Delta + A0Delta, A0Delta + A0Delta, A0Delta + A0Delta);
+			VectorRegister4Float Z1DataVec = MakeVectorRegister(Z1Data[0], Z1Data[1], Z1Data[0], Z1Data[1]);
+
+			if (NumToSimd)
+			{
+				A0CurrVec = VectorAdd(A0CurrVec, MakeVectorRegister(A0Delta, A0Delta, A0Delta + A0Delta, A0Delta + A0Delta));	//a1 + delta, a1 + delta, a1 + 2x delta, a1 + 2x delta
+			}
+
+			for (; SampleIndex < NumToSimd; SampleIndex += AUDIO_NUM_FLOATS_PER_VECTOR_REGISTER)
+			{
+				VectorRegister4Float InputSample = VectorLoad(&InputBuffer[SampleIndex]);
+				VectorRegister4Float Vn = VectorMultiply(VectorSubtract(InputSample, Z1DataVec), A0CurrVec);	//Vn1, Vn2, garbage, garbage
+				VectorRegister4Float LPF = VectorAdd(Vn, Z1DataVec);	//LPF1, LPF2, garbage, garbage
+
+				VectorRegister4Float VnLPFSum = VectorAdd(Vn, LPF);
+				Z1DataVec = VectorShuffle(Z1DataVec, VnLPFSum, 0, 1, 0, 1);
+				Vn = VectorMultiply(VectorSubtract(InputSample, Z1DataVec), A0CurrVec);	//Vn1, Vn2, Vn3, Vn4
+				LPF = VectorAdd(Vn, Z1DataVec);	//LPF1, LPF2, LPF3, LPF4
+
+				VnLPFSum = VectorAdd(Vn, LPF);
+				Z1DataVec = VectorSwizzle(VnLPFSum, 2, 3, 2, 3);	//Vn2, Vn3, Vn2, Vn3
+
+				VectorStore(VectorSubtract(InputSample, LPF), &OutputBuffer[SampleIndex]);
+
+				A0CurrVec = VectorAdd(A0CurrVec, A0DeltaVec);	//a1 + 2x delta, a1 + 2x delta, a1 + 2x delta, a1 + 2x delta
+			}
+		}
+		else if (NumChannels == 4)
+		{
+			VectorRegister4Float A0CurrVec = VectorLoadFloat1(&A0Curr);
+			VectorRegister4Float A0DeltaVec = VectorLoadFloat1(&A0Delta);
+			VectorRegister4Float Z1DataVec = VectorLoad(&Z1Data[0]);
+
+			for (; SampleIndex < NumToSimd; SampleIndex += AUDIO_NUM_FLOATS_PER_VECTOR_REGISTER)
+			{
+				A0CurrVec = VectorAdd(A0CurrVec, A0DeltaVec);
+
+				VectorRegister4Float InputSample = VectorLoad(&InputBuffer[SampleIndex]);
+				VectorRegister4Float Vn = VectorMultiply(VectorSubtract(InputSample, Z1DataVec), A0CurrVec);
+				VectorRegister4Float LPF = VectorAdd(Vn, Z1DataVec);
+
+				Z1DataVec = VectorAdd(Vn, LPF);
+
+				VectorStore(VectorSubtract(InputSample, LPF), &OutputBuffer[SampleIndex]);
+			}
+
+			float Store[4];
+			VectorStore(A0CurrVec, Store);
+			A0Curr = Store[0];
+			VectorStore(Z1DataVec, &Z1Data[0]);
+		}
+
+		for (; SampleIndex < NumSamples; ++SampleIndex)
 		{
 			// cache which delay term we should be using
 			const int32 ChannelIndex = SampleIndex % NumChannels;
