Remove API markings on enums and interfaces that don't need it to fix modular build on some platforms

#rnx

diff --git a/Engine/Plugins/Experimental/Fracture/Source/FractureEngine/Public/FractureEngineClustering.h b/Engine/Plugins/Experimental/Fracture/Source/FractureEngine/Public/FractureEngineClustering.h
--- a/Engine/Plugins/Experimental/Fracture/Source/FractureEngine/Public/FractureEngineClustering.h
+++ b/Engine/Plugins/Experimental/Fracture/Source/FractureEngine/Public/FractureEngineClustering.h
@@ -70,7 +70,7 @@ private:
 	TArray<bool> Visited;
 };
 
-enum class FRACTUREENGINE_API EFractureEngineClusterSizeMethod : uint8
+enum class EFractureEngineClusterSizeMethod : uint8
 {
 	// Cluster by specifying an absolute number of clusters
 	ByNumber,
diff --git a/Engine/Plugins/Runtime/MeshModelingToolset/Source/ModelingComponents/Public/AssetUtils/MeshDescriptionUtil.h b/Engine/Plugins/Runtime/MeshModelingToolset/Source/ModelingComponents/Public/AssetUtils/MeshDescriptionUtil.h
--- a/Engine/Plugins/Runtime/MeshModelingToolset/Source/ModelingComponents/Public/AssetUtils/MeshDescriptionUtil.h
+++ b/Engine/Plugins/Runtime/MeshModelingToolset/Source/ModelingComponents/Public/AssetUtils/MeshDescriptionUtil.h
@@ -46,7 +46,7 @@ namespace UE
 		//
 
 		/** How to update a boolean build setting */
-		enum class MODELINGCOMPONENTS_API EBuildSettingBoolChange
+		enum class EBuildSettingBoolChange
 		{
 			Disable,
 			Enable,
diff --git a/Engine/Source/Runtime/Experimental/Chaos/Public/Chaos/PBDRigidClusteringAlgo.h b/Engine/Source/Runtime/Experimental/Chaos/Public/Chaos/PBDRigidClusteringAlgo.h
--- a/Engine/Source/Runtime/Experimental/Chaos/Public/Chaos/PBDRigidClusteringAlgo.h
+++ b/Engine/Source/Runtime/Experimental/Chaos/Public/Chaos/PBDRigidClusteringAlgo.h
@@ -5,7 +5,7 @@
 
 namespace Chaos
 {
-	enum class CHAOS_API EMassOffsetType : uint8
+	enum class EMassOffsetType : uint8
 	{
 		None = 0,
 		Position = 1 << 0,
@@ -13,7 +13,7 @@ namespace Chaos
 	};
 	ENUM_CLASS_FLAGS(EMassOffsetType)
 
-	enum class CHAOS_API EInertiaOperations : uint8
+	enum class EInertiaOperations : uint8
 	{
 		None = 0,
 		LocalInertiaDropOffDiagonalTerms = 1 << 0
diff --git a/Engine/Source/Runtime/SignalProcessing/Public/DSP/BufferDiagnostics.h b/Engine/Source/Runtime/SignalProcessing/Public/DSP/BufferDiagnostics.h
--- a/Engine/Source/Runtime/SignalProcessing/Public/DSP/BufferDiagnostics.h
+++ b/Engine/Source/Runtime/SignalProcessing/Public/DSP/BufferDiagnostics.h
@@ -51,7 +51,7 @@ namespace Audio
 
 	SIGNALPROCESSING_API bool ContainsNan(TArrayView<const float> InBuffer);
 
-	enum class SIGNALPROCESSING_API ECheckBufferFlags : uint32
+	enum class ECheckBufferFlags : uint32
 	{
 		None = 0,
 		Infinity = 1 << 1,
@@ -86,7 +86,7 @@ namespace Audio
 	**/
 	SIGNALPROCESSING_API bool CheckBuffer(TArrayView<const float> InBuffer, const ECheckBufferFlags InFlags, ECheckBufferFlags& OutFailedFlags);
 
-	enum class SIGNALPROCESSING_API EBufferCheckBehavior : uint8
+	enum class EBufferCheckBehavior : uint8
 	{
 		Nothing,
 		Ensure,
