Fixed vectorized version of FInterpolatedHPF::ProcessAudioBuffer where I forgot to store content of vector registers into class members

diff --git a/Engine/Source/Runtime/SignalProcessing/Private/InterpolatedOnePole.cpp b/Engine/Source/Runtime/SignalProcessing/Private/InterpolatedOnePole.cpp
--- a/Engine/Source/Runtime/SignalProcessing/Private/InterpolatedOnePole.cpp
+++ b/Engine/Source/Runtime/SignalProcessing/Private/InterpolatedOnePole.cpp
@@ -409,6 +409,13 @@ namespace Audio
 
 				A0CurrVec = VectorAdd(A0CurrVec, A0DeltaVec);	//a1 + 2x delta, a1 + 2x delta, a1 + 2x delta, a1 + 2x delta
 			}
+
+			float Store[4];
+			VectorStore(A0CurrVec, Store);
+			A0Curr = Store[0];
+			VectorStore(Z1DataVec, Store);
+			Z1Data[0] = Store[0];
+			Z1Data[1] = Store[1];
 		}
 		else if (NumChannels == 4)
 		{
