Fix for City Sample editor crash in Audio::ArrayDeinterleave()

diff --git a/Engine/Source/Runtime/SignalProcessing/Private/DynamicsProcessor.cpp b/Engine/Source/Runtime/SignalProcessing/Private/DynamicsProcessor.cpp
--- a/Engine/Source/Runtime/SignalProcessing/Private/DynamicsProcessor.cpp
+++ b/Engine/Source/Runtime/SignalProcessing/Private/DynamicsProcessor.cpp
@@ -338,7 +338,7 @@ namespace Audio
 			DeinterleaveOutput[ChannelIndex] = &OutputWorkBuffer[ChannelIndex * NumSubbufferFrames];
 			DeinterleaveEnvelope[ChannelIndex] = &EnvelopeWorkBuffer[ChannelIndex * NumSubbufferFrames];
 		};
-		for (int32 KeyChannelIndex = 0; KeyChannelIndex < NumChannels; KeyChannelIndex++)
+		for (int32 KeyChannelIndex = 0; KeyChannelIndex < KeyNumChannels; KeyChannelIndex++)
 		{
 			DeinterleaveKey[KeyChannelIndex] = &KeyWorkBuffer[KeyChannelIndex * NumSubbufferFrames];
 		}
