ISPC Audio Optimizations by [REDACTED_EMAIL]

1. ArrayMax
2. ArrayMaxAbsValue
3. ArrayComplexMultiplyAdd
4. ArrayFloatToPcm16
5. ArrayPcm16ToFloat
6. FLinearInterpFractionalDelay::ProcessAudioBlock
7. FLongDelayAPF::ProcessAudioBlock
8. FVectorComplexFFT::ScaledComplexConjugate

diff --git a/Engine/Source/Runtime/SignalProcessing/Private/FloatArrayMath.cpp b/Engine/Source/Runtime/SignalProcessing/Private/FloatArrayMath.cpp
--- a/Engine/Source/Runtime/SignalProcessing/Private/FloatArrayMath.cpp
+++ b/Engine/Source/Runtime/SignalProcessing/Private/FloatArrayMath.cpp
@@ -5,6 +5,7 @@
 #include "Math/UnrealMathVectorConstants.h"
 #include "SignalProcessingModule.h"
 #include "ProfilingDebugging/CsvProfiler.h"
+#include "DSP/Dsp.h"
 
 #if INTEL_ISPC && !UE_BUILD_SHIPPING
 #include "HAL/IConsoleManager.h"
@@ -744,24 +745,34 @@ namespace Audio
 
 		CSV_SCOPED_TIMING_STAT(Audio_Dsp, ArrayMax);
 		const int32 Num = InView1.Num();
-		const int32 NumToSimd = Num & MathIntrinsics::SimdMask;
-		const int32 NumNotToSimd = Num & MathIntrinsics::NotSimdMask;
-
-		if (NumToSimd)
+		
+		if (bAudio_FloatArrayMath_ISPC_Enabled)
 		{
-			for (int32 i = 0; i < NumToSimd; i += AUDIO_NUM_FLOATS_PER_VECTOR_REGISTER)
+#if INTEL_ISPC
+			ispc::ArrayMax(InView1.GetData(), InView2.GetData(), OutView.GetData(), Num);
+#endif
+		}
+		else
+		{
+			const int32 NumToSimd = Num & MathIntrinsics::SimdMask;
+			const int32 NumNotToSimd = Num & MathIntrinsics::NotSimdMask;
+
+			if (NumToSimd)
 			{
-				VectorRegister4Float Input1 = VectorLoad(&InView1[i]);
-				VectorRegister4Float Input2 = VectorLoad(&InView2[i]);
-				VectorStore(VectorMax(Input1, Input2), &OutView[i]);
+				for (int32 i = 0; i < NumToSimd; i += AUDIO_NUM_FLOATS_PER_VECTOR_REGISTER)
+				{
+					VectorRegister4Float Input1 = VectorLoad(&InView1[i]);
+					VectorRegister4Float Input2 = VectorLoad(&InView2[i]);
+					VectorStore(VectorMax(Input1, Input2), &OutView[i]);
+				}
 			}
-		}
 
-		if (NumNotToSimd)
-		{
-			for (int32 i = NumToSimd; i < Num; ++i)
+			if (NumNotToSimd)
 			{
-				OutView[i] = FMath::Max(InView1[i], InView2[i]);
+				for (int32 i = NumToSimd; i < Num; ++i)
+				{
+					OutView[i] = FMath::Max(InView1[i], InView2[i]);
+				}
 			}
 		}
 	}
@@ -773,35 +784,44 @@ namespace Audio
 		const int32 Num = InView.Num();
 		const float* Data = InView.GetData();
 
-		float Max = 0.f;
-		const int32 NumToSimd = Num & MathIntrinsics::SimdMask;
-		const int32 NumNotToSimd = Num & MathIntrinsics::NotSimdMask;
-
-		VectorRegister4Float MaxVector = VectorSetFloat1(0.f);
-		
-		if (NumToSimd)
+		if (bAudio_FloatArrayMath_ISPC_Enabled)
 		{
-			for (int32 i = 0; i < NumToSimd; i += AUDIO_NUM_FLOATS_PER_VECTOR_REGISTER)
+#if INTEL_ISPC
+			return ispc::ArrayMaxAbsValue(Data, Num);
+#endif
+		}
+		else
+		{
+			float Max = 0.f;
+			const int32 NumToSimd = Num & MathIntrinsics::SimdMask;
+			const int32 NumNotToSimd = Num & MathIntrinsics::NotSimdMask;
+
+			VectorRegister4Float MaxVector = VectorSetFloat1(0.f);
+
+			if (NumToSimd)
 			{
-				VectorRegister4Float Input1 = VectorLoad(&Data[i]);
+				for (int32 i = 0; i < NumToSimd; i += AUDIO_NUM_FLOATS_PER_VECTOR_REGISTER)
+				{
+					VectorRegister4Float Input1 = VectorLoad(&Data[i]);
 
-				MaxVector = VectorMax(MaxVector, VectorAbs(Input1));
-			}
+					MaxVector = VectorMax(MaxVector, VectorAbs(Input1));
+				}
 
-			AlignedFloat4 OutArray(MaxVector);
-			
-			Max =  FMath::Max(FMath::Max(OutArray[0], OutArray[1]), FMath::Max(OutArray[2], OutArray[3]));
-		}
+				AlignedFloat4 OutArray(MaxVector);
 
-		if (NumNotToSimd)
-		{
-			for (int32 i = NumToSimd; i < Num; ++i)
+				Max = FMath::Max(FMath::Max(OutArray[0], OutArray[1]), FMath::Max(OutArray[2], OutArray[3]));
+			}
+
+			if (NumNotToSimd)
 			{
-				Max = FMath::Max(FMath::Abs(Data[i]), Max);
+				for (int32 i = NumToSimd; i < Num; ++i)
+				{
+					Max = FMath::Max(FMath::Abs(Data[i]), Max);
+				}
 			}
-		}
 
-		return Max;
+			return Max;
+		}
 	}
 
 	void ArrayMultiply(TArrayView<const float> InFloatBufferA, TArrayView<const float> InFloatBufferB, TArrayView<float> OutBuffer)
@@ -935,6 +955,86 @@ namespace Audio
 		}
 	}
 
+	void ArrayComplexMultiplyAdd(TArrayView<const float> InValues1, TArrayView<const float> InValues2, TArrayView<float> OutArray)
+	{
+		CSV_SCOPED_TIMING_STAT(Audio_Dsp, ArrayComplexMultiplyAdd);
+
+		check(InValues1.Num() == InValues2.Num());
+		check(OutArray.Num() == InValues1.Num());
+
+		const int32 Num = InValues1.Num();
+
+		const float* InAData = InValues1.GetData();
+		const float* InBData = InValues2.GetData();
+		float* OutData = OutArray.GetData();
+
+		if (bAudio_FloatArrayMath_ISPC_Enabled)
+		{
+#if INTEL_ISPC
+			ispc::ArrayComplexMultiplyAdd(InAData, InBData, OutData, Num);
+#endif
+		}
+		else
+		{
+			const int32 NumSimd = Num & MathIntrinsics::SimdMask;
+
+			// Complex numbers are stored as [real_0, complex_0, real_1, complex_1, ... real_N, complex_N]
+			// So we final amount must be evenly divisble by 2.
+			check(NumSimd % 2 == 0);
+
+			const VectorRegister4Float SignFlip = MakeVectorRegisterFloat(-1.f, 1.f, -1.f, 1.f);
+
+			for (int32 i = 0; i < NumSimd; i += 4)
+			{
+				// Complex multiply add
+				// Nr = real component of Nth number
+				// Ni = imaginary component of Nth number
+				//
+				//
+				// The input is then 
+				// A1r A1i A2r A2i
+				// B1r B1i B2r B2i
+
+				// VectorA = A1r A1i A2r A2i
+				VectorRegister4Float VectorInA = VectorLoad(&InAData[i]);
+				// Temp12 = A1i A1r A2i A2r
+				VectorRegister4Float Temp1 = VectorSwizzle(VectorInA, 1, 0, 3, 2);
+
+				// VectorB = B1r B1i B2r B2i
+				VectorRegister4Float VectorInB = VectorLoad(&InBData[i]);
+				// Temp2 = B1r B1r B2r B2r
+				VectorRegister4Float Temp2 = VectorSwizzle(VectorInB, 0, 0, 2, 2);
+				// Temp3 = B1i B1i B2i B2i
+				VectorRegister4Float Temp3 = VectorSwizzle(VectorInB, 1, 1, 3, 3);
+
+
+				// VectorA = A1rB1r, A1iB1r, A2rB2r, A2iB2r
+				VectorInA = VectorMultiply(VectorInA, Temp2);
+
+				// Temp1 = A1iB1i, A1rB1i, A2iB2i, A2rb2i
+				Temp1 = VectorMultiply(Temp1, Temp3);
+
+				// Temp1 = -A1iB1i, A1rB1i, -A2iB2i, A2rb2i
+				// Temp1 = A1rB1r - A1iB1i, A1iB1r + A1rB1i, A2rB2r - A2iB2i, A2iB2r + A2rB2i
+				Temp1 = VectorMultiplyAdd(Temp1, SignFlip, VectorInA);
+
+				// VectorOut = O1r + A1rB1r - A1iB1i, O1i + A1iB1r + A1rB1i, O2r + A2rB2r - A2iB2i, O2i + A2iB2r + A2rB2i
+				VectorRegister4Float VectorOut = VectorLoad(&OutData[i]);
+				VectorOut = VectorAdd(Temp1, VectorOut);
+
+				VectorStore(VectorOut, &OutData[i]);
+			}
+
+			for (int32 i = NumSimd; i < Num; i += 2)
+			{
+				// Real output
+				OutData[i] += (InAData[i] * InBData[i]) - (InAData[i + 1] * InBData[i + 1]);
+				// Imaginary output
+				OutData[i + 1] += (InAData[i + 1] * InBData[i]) + (InAData[i] * InBData[i + 1]);
+			}
+		}
+	}
+
 	void ArrayMultiplyByConstant(TArrayView<const float> InFloatBuffer, float InValue, TArrayView<float> OutFloatBuffer)
 	{
 		CSV_SCOPED_TIMING_STAT(Audio_Dsp, ArrayMultiplyByConstant);
@@ -2327,41 +2427,50 @@ namespace Audio
 		const float* InputPtr = InView.GetData();
 		int16* OutPtr = OutView.GetData();
 
-		constexpr float ConversionValue = static_cast<float>(TNumericLimits<int16>::Max());
-		const VectorRegister4Float Multiplier = VectorSetFloat1(ConversionValue);
-		int32 i = 0;
-#if PLATFORM_ENABLE_VECTORINTRINSICS_NEON
-		const int32 SimdNum = Num & MathIntrinsics::Simd8Mask;
-		for (; i < SimdNum; i += 8)
+		if (bAudio_FloatArrayMath_ISPC_Enabled)
 		{
-			const float32x4x2_t InVector = vld1q_f32_x2(&InputPtr[i]);
-			const VectorRegister4Float ScaledVector1 = VectorMultiply(InVector.val[0], Multiplier);
-			const VectorRegister4Float ScaledVector2 = VectorMultiply(InVector.val[1], Multiplier);
-			const VectorRegister4Int IntVector1 = VectorFloatToInt(ScaledVector1);
-			const VectorRegister4Int IntVector2 = VectorFloatToInt(ScaledVector2);
-			const int16x8_t Result = vmovn_high_s32(vmovn_u32(IntVector1), IntVector2);
-			vst1q_s16(&OutPtr[i], Result);
+#if INTEL_ISPC
+			ispc::ArrayFloatToPcm16(InputPtr, OutPtr, Num);
+#endif
 		}
-#else
-		const int32 SimdNum = Num & MathIntrinsics::SimdMask;
-		for (; i < SimdNum; i += AUDIO_NUM_FLOATS_PER_VECTOR_REGISTER)
+		else
 		{
-			const VectorRegister4Float InVector = VectorLoad(&InputPtr[i]);
-			const VectorRegister4Float ScaledVector = VectorMultiply(InVector, Multiplier);
-			const VectorRegister4Int IntVector = VectorFloatToInt(ScaledVector);
+			constexpr float ConversionValue = static_cast<float>(TNumericLimits<int16>::Max());
+			const VectorRegister4Float Multiplier = VectorSetFloat1(ConversionValue);
+			int32 i = 0;
+#if PLATFORM_ENABLE_VECTORINTRINSICS_NEON
+			const int32 SimdNum = Num & MathIntrinsics::Simd8Mask;
+			for (; i < SimdNum; i += 8)
+			{
+				const float32x4x2_t InVector = vld1q_f32_x2(&InputPtr[i]);
+				const VectorRegister4Float ScaledVector1 = VectorMultiply(InVector.val[0], Multiplier);
+				const VectorRegister4Float ScaledVector2 = VectorMultiply(InVector.val[1], Multiplier);
+				const VectorRegister4Int IntVector1 = VectorFloatToInt(ScaledVector1);
+				const VectorRegister4Int IntVector2 = VectorFloatToInt(ScaledVector2);
+				const int16x8_t Result = vmovn_high_s32(vmovn_u32(IntVector1), IntVector2);
+				vst1q_s16(&OutPtr[i], Result);
+			}
+#else
+			const int32 SimdNum = Num & MathIntrinsics::SimdMask;
+			for (; i < SimdNum; i += AUDIO_NUM_FLOATS_PER_VECTOR_REGISTER)
+			{
+				const VectorRegister4Float InVector = VectorLoad(&InputPtr[i]);
+				const VectorRegister4Float ScaledVector = VectorMultiply(InVector, Multiplier);
+				const VectorRegister4Int IntVector = VectorFloatToInt(ScaledVector);
 
-			const AlignedFloat4 ScaledFloatArray(ScaledVector);
+				const AlignedFloat4 ScaledFloatArray(ScaledVector);
 
-			OutPtr[i + 0] = (int16)ScaledFloatArray[0];
-			OutPtr[i + 1] = (int16)ScaledFloatArray[1];
-			OutPtr[i + 2] = (int16)ScaledFloatArray[2];
-			OutPtr[i + 3] = (int16)ScaledFloatArray[3];
-		}
+				OutPtr[i + 0] = (int16)ScaledFloatArray[0];
+				OutPtr[i + 1] = (int16)ScaledFloatArray[1];
+				OutPtr[i + 2] = (int16)ScaledFloatArray[2];
+				OutPtr[i + 3] = (int16)ScaledFloatArray[3];
+			}
 #endif //~PLATFORM_ENABLE_VECTORINTRINSICS_NEON
 
-		for (; i < Num; i++)
-		{
-			OutPtr[i] = (int16)(InputPtr[i] * ConversionValue);
+			for (; i < Num; i++)
+			{
+				OutPtr[i] = (int16)(InputPtr[i] * ConversionValue);
+			}
 		}
 	}
 	
@@ -2376,42 +2485,51 @@ namespace Audio
 		const int16* InputPtr = InView.GetData();
 		float* OutPtr = OutView.GetData();
 
-		constexpr float ConversionValue = 1.f / static_cast<float>(TNumericLimits<int16>::Max());
-		const VectorRegister4Float Multiplier = VectorSetFloat1(ConversionValue);
-
-		int32 i = 0;
-#if PLATFORM_ENABLE_VECTORINTRINSICS_NEON
-		const int32 SimdNum = Num & MathIntrinsics::Simd8Mask;
-		for (; i < SimdNum; i += 8)
+		if (bAudio_FloatArrayMath_ISPC_Enabled)
 		{
-			int16x8_t Data = vld1q_s16(&InputPtr[i]);
-			int32x4_t VecA = vmovl_s16(vget_low_s16(Data));
-			int32x4_t VecB = vmovl_high_s16(Data);
-			float32x4x2_t FloatVec;
-			FloatVec.val[0] = VectorMultiply(vcvtq_f32_s32(VecA), Multiplier);
-			FloatVec.val[1] = VectorMultiply(vcvtq_f32_s32(VecB), Multiplier);
-			vst1q_f32_x2(&OutPtr[i], FloatVec);
+#if INTEL_ISPC
+			ispc::ArrayPcm16ToFloat(InputPtr, OutPtr, Num);
+#endif
 		}
-#else
-		AlignedFloat4 FloatArray(GlobalVectorConstants::FloatZero);
-		const int32 SimdNum = Num & MathIntrinsics::SimdMask;
-		for (; i < SimdNum; i += AUDIO_NUM_FLOATS_PER_VECTOR_REGISTER)
+		else
 		{
-			FloatArray[0] = (float)InputPtr[i];
-			FloatArray[1] = (float)InputPtr[i + 1];
-			FloatArray[2] = (float)InputPtr[i + 2];
-			FloatArray[3] = (float)InputPtr[i + 3];
+			constexpr float ConversionValue = 1.f / static_cast<float>(TNumericLimits<int16>::Max());
+			const VectorRegister4Float Multiplier = VectorSetFloat1(ConversionValue);
 
-			const VectorRegister4Float InVector = FloatArray.ToVectorRegister();
-			const VectorRegister4Float ScaledVector = VectorMultiply(InVector, Multiplier);
+			int32 i = 0;
+#if PLATFORM_ENABLE_VECTORINTRINSICS_NEON
+			const int32 SimdNum = Num & MathIntrinsics::Simd8Mask;
+			for (; i < SimdNum; i += 8)
+			{
+				int16x8_t Data = vld1q_s16(&InputPtr[i]);
+				int32x4_t VecA = vmovl_s16(vget_low_s16(Data));
+				int32x4_t VecB = vmovl_high_s16(Data);
+				float32x4x2_t FloatVec;
+				FloatVec.val[0] = VectorMultiply(vcvtq_f32_s32(VecA), Multiplier);
+				FloatVec.val[1] = VectorMultiply(vcvtq_f32_s32(VecB), Multiplier);
+				vst1q_f32_x2(&OutPtr[i], FloatVec);
+			}
+#else
+			AlignedFloat4 FloatArray(GlobalVectorConstants::FloatZero);
+			const int32 SimdNum = Num & MathIntrinsics::SimdMask;
+			for (; i < SimdNum; i += AUDIO_NUM_FLOATS_PER_VECTOR_REGISTER)
+			{
+				FloatArray[0] = (float)InputPtr[i];
+				FloatArray[1] = (float)InputPtr[i + 1];
+				FloatArray[2] = (float)InputPtr[i + 2];
+				FloatArray[3] = (float)InputPtr[i + 3];
 
-			VectorStore(ScaledVector, &OutPtr[i]);
-		}
+				const VectorRegister4Float InVector = FloatArray.ToVectorRegister();
+				const VectorRegister4Float ScaledVector = VectorMultiply(InVector, Multiplier);
+
+				VectorStore(ScaledVector, &OutPtr[i]);
+			}
 #endif //~PLATFORM_ENABLE_VECTORINTRINSICS_NEON
 
-		for (; i < Num; i++)
-		{
-			OutPtr[i] = (float)InputPtr[i] * ConversionValue;
+			for (; i < Num; i++)
+			{
+				OutPtr[i] = (float)InputPtr[i] * ConversionValue;
+			}
 		}
 	}
 	
@@ -2433,8 +2551,6 @@ namespace Audio
         
         constexpr int32 SizeofPCM24 = 3;
         constexpr float ConversionValue = static_cast<float>(0x7fffff);
-        // Output byte order must adhere to host endianess
-        const bool bIsOutputLittleEndian = IsHostLittleEndian();
 
         constexpr int32 SimdIndexStride = (4 * SizeofPCM24);
         check((OutView.Num() / SizeofPCM24) >= InView.Num());
@@ -2451,23 +2567,11 @@ namespace Audio
 		constexpr uint32 ZeroFill = 0x80;
 		// Here we convert from 4 32-bit ints to 4 24-bit ints using the mask register below.
 		// Notice that every 4th byte is skipped. This is the uneeded high byte of the 32-bit int.
-		VectorRegister4Int Mask;
-		if (bIsOutputLittleEndian)
-		{
-			Mask = MakeVectorRegisterInt(
-				CreateByteMask(0, 1, 2, 4),
-				CreateByteMask(5, 6, 8, 9),
-				CreateByteMask(10, 12, 13, 14),
-				CreateByteMask(ZeroFill, ZeroFill, ZeroFill, ZeroFill));
-		}
-		else
-		{
-			Mask = MakeVectorRegisterInt(
-				CreateByteMask(2, 1, 0, 6),
-				CreateByteMask(5, 4, 10, 9),
-				CreateByteMask(8, 14, 13, 12),
-				CreateByteMask(ZeroFill, ZeroFill, ZeroFill, ZeroFill));
-		}
+		VectorRegister4Int Mask = MakeVectorRegisterInt(
+			CreateByteMask(0, 1, 2, 4),
+			CreateByteMask(5, 6, 8, 9),
+			CreateByteMask(10, 12, 13, 14),
+			CreateByteMask(ZeroFill, ZeroFill, ZeroFill, ZeroFill));
 
 		for (; InIndex < SimdNum; InIndex += AUDIO_NUM_FLOATS_PER_VECTOR_REGISTER)
 		{
@@ -2482,29 +2586,14 @@ namespace Audio
         
 		for (; InIndex < Num; InIndex++)
 		{
-			const int32 ConvertedValue = InputPtr[InIndex] * ConversionValue;
-
-			AssignPcm24Value(&OutPtr[OutIndex], ConvertedValue, bIsOutputLittleEndian);
-			OutIndex += SizeofPCM24;
-		}
-	}
+			const uint32 ConvertedValue = uint32(InputPtr[InIndex] * ConversionValue);
+			uint8* UnsignedOutPtr = (uint8*)&OutPtr[OutIndex];
 
-	FORCEINLINE void AssignPcm24Value(int8* OutPtr, const int32 InValue, const bool bIsOutputLittleEndian)
-	{
-		uint8* UnsignedOutPtr = (uint8*)OutPtr;
-		const uint32 UnsignedInValue = (const uint32)InValue;
+			UnsignedOutPtr[0] = ConvertedValue & 0xFF;
+			UnsignedOutPtr[1] = ConvertedValue >> 8 & 0xFF;
+			UnsignedOutPtr[2] = ConvertedValue >> 16 & 0xFF;
 
-		if (bIsOutputLittleEndian)
-		{
-			UnsignedOutPtr[0] = UnsignedInValue & 0xFF;
-			UnsignedOutPtr[1] = UnsignedInValue >> 8 & 0xFF;
-			UnsignedOutPtr[2] = UnsignedInValue >> 16 & 0xFF;
-		}
-		else
-		{
-			UnsignedOutPtr[2] = UnsignedInValue & 0xFF;
-			UnsignedOutPtr[1] = UnsignedInValue >> 8 & 0xFF;
-			UnsignedOutPtr[0] = UnsignedInValue >> 16 & 0xFF;
+			OutIndex += SizeofPCM24;
 		}
 	}
 
@@ -2648,12 +2737,12 @@ namespace Audio
 	{
 		CSV_SCOPED_TIMING_STAT(Audio_Dsp, ArrayDeinterleave);
 
-		for(int32 ChannelIdx = 0; ChannelIdx < InChannels; ChannelIdx++)
+		for (int32 ChannelIdx = 0; ChannelIdx < InChannels; ChannelIdx++)
 		{
 			const float* InPtr = &InBuffer[ChannelIdx];
 			float* OutPtr = OutBuffers[ChannelIdx];
-			
-			for(int32 SampleIdx = 0; SampleIdx < InFrames; SampleIdx++)
+
+			for (int32 SampleIdx = 0; SampleIdx < InFrames; SampleIdx++)
 			{
 				*OutPtr++ = *InPtr;
 				InPtr += InChannels;
@@ -2892,6 +2981,152 @@ namespace Audio
 		}
 	}
 
+	void ArrayAPFLongDelayProcess(const float* InSamples, const float* InDelaySamples, const int32 InNum, float* OutSamples, float* OutDelaySamples, const float Gain)
+	{
+		if (bAudio_FloatArrayMath_ISPC_Enabled)
+		{
+#if INTEL_ISPC
+			ispc::ArrayAPFLongDelayProcess(InSamples,
+				InDelaySamples,
+				InNum,
+				OutSamples,
+				OutDelaySamples,
+				Gain);
+#endif
+		}
+		else
+		{
+			// Calculate new delay line samples. "w[n] = x[n] + gw[n - d]"
+			int32 NumToSIMD = InNum - (InNum % AUDIO_NUM_FLOATS_PER_VECTOR_REGISTER);
+
+			VectorRegister4Float VG = MakeVectorRegisterFloat(Gain, Gain, Gain, Gain);
+			VectorRegister4Float VNG = MakeVectorRegisterFloat(-Gain, -Gain, -Gain, -Gain);
+
+			for (int32 i = 0; i < InNum; i += AUDIO_NUM_FLOATS_PER_VECTOR_REGISTER)
+			{
+				VectorRegister4Float VInDelay = VectorLoadAligned(&InDelaySamples[i]);
+				VectorRegister4Float VInSamples = VectorLoadAligned(&InSamples[i]);
+				// w[n] = x[n] + G * w[n - D]
+				VectorRegister4Float VOutDelay = VectorMultiplyAdd(VInDelay, VG, VInSamples);
+				VectorStoreAligned(VOutDelay, &OutDelaySamples[i]);
+
+				// y[n] = -G * w[n] + w[n - D]
+				VectorRegister4Float VOut = VectorMultiplyAdd(VOutDelay, VNG, VInDelay);
+				VectorStoreAligned(VOut, &OutSamples[i]);
+			}
+
+			// Calculate allpass for remaining samples that we couldn't SIMD
+			for (int32 i = NumToSIMD; i < InNum; i++)
+			{
+				OutDelaySamples[i] = InDelaySamples[i] * Gain + InSamples[i];
+				OutSamples[i] = OutDelaySamples[i] * -Gain + InDelaySamples[i];
+			}
+		}
+	}
+
+	void ArrayLerpFractionalDelay(const float* InSamples, const float* InDelays, const float* DelayData, const int* IntegerDelays, int* UpperDelayPos, int* LowerDelayPos, const int32 InNum, float* OutSamples, const float MaxDelay)
+	{
+		if (bAudio_FloatArrayMath_ISPC_Enabled)
+		{
+#if INTEL_ISPC
+			ispc::ArrayLerpFractionalDelay(InSamples,
+				InDelays,
+				DelayData,
+				IntegerDelays,
+				InNum,
+				OutSamples,
+				MaxDelay);
+#endif
+		}
+		else
+		{
+			const VectorRegister4Float VMaxDelay = MakeVectorRegister(MaxDelay, MaxDelay, MaxDelay, MaxDelay);
+			for (int32 i = 0; i < InNum; i += 4)
+			{
+				VectorRegister4Float VFractionalDelays = VectorLoad(&InDelays[i]);
+				// Ensure fractional delays are positive
+				VFractionalDelays = VectorMax(VFractionalDelays, GlobalVectorConstants::FloatZero);
+				VFractionalDelays = VectorMin(VFractionalDelays, VMaxDelay);
+
+				// Separate integer from fraction
+				VectorRegister4Float VFloorDelays = VectorFloor(VFractionalDelays);
+
+				// Determine linear weights
+				VectorRegister4Float VUpperCoefficients = VectorSubtract(VFractionalDelays, VFloorDelays);
+				VectorRegister4Float VLowerCoefficients = VectorSubtract(GlobalVectorConstants::FloatOne, VUpperCoefficients);
+
+
+				// Make integer locations relative to block
+				VectorRegister4Int VIntegerDelays = VectorFloatToInt(VFloorDelays);
+				VectorRegister4Int VIntegerDelayOffset = VectorIntLoadAligned(&IntegerDelays[i]);
+				VIntegerDelays = VectorIntSubtract(VIntegerDelayOffset, VIntegerDelays);
+
+				// Lookup samples for interpolation
+				VectorIntStoreAligned(VIntegerDelays, UpperDelayPos);
+				VectorIntStoreAligned(VectorIntAdd(VIntegerDelays, GlobalVectorConstants::IntOne), LowerDelayPos);
+
+				VectorRegister4Float VLowerSamples = MakeVectorRegister(
+					DelayData[LowerDelayPos[0]],
+					DelayData[LowerDelayPos[1]],
+					DelayData[LowerDelayPos[2]],
+					DelayData[LowerDelayPos[3]]
+				);
+				VectorRegister4Float VUpperSamples = MakeVectorRegister(
+					DelayData[UpperDelayPos[0]],
+					DelayData[UpperDelayPos[1]],
+					DelayData[UpperDelayPos[2]],
+					DelayData[UpperDelayPos[3]]
+				);
+
+				// Interpolate samples
+				VectorRegister4Float VOut = VectorMultiplyAdd(
+					VLowerSamples,
+					VLowerCoefficients,
+					VectorMultiply(VUpperSamples, VUpperCoefficients));
+				VectorStore(VOut, &OutSamples[i]);
+			}
+		}
+	}
+
+	void ArrayScaledComplexConjugate(const float* RESTRICT InValues, const int32 Num, float* RESTRICT OutValues, const float Scale)
+	{
+		if (bAudio_FloatArrayMath_ISPC_Enabled)
+		{
+#if INTEL_ISPC
+			ispc::ArrayScaledComplexConjugate(InValues,
+				Num,
+				OutValues,
+				Scale);
+#endif
+		}
+		else
+		{
+			// Use mask to quickly find out number of values that can be SIMD'd
+			const int32 SIMD_MASK = 0xFFFFFFFC;
+			const int32 NumToSimd = SIMD_MASK & Num;
+
+			// Complex values in a vector are [real_1, imag_1, real_2, imag_2].  
+			// By multipling this value, we flip the sign of the imaginary components, which
+			// is the equivalent of a complex conjugate.
+			const VectorRegister4Float SignFlipImag = MakeVectorRegisterFloat(Scale, -Scale, Scale, -Scale);
+
+			// Perform operation using SIMD
+			for (int32 i = 0; i < NumToSimd; i += 4)
+			{
+				VectorRegister4Float Value = VectorLoad(&InValues[i]);
+				Value = VectorMultiply(SignFlipImag, Value);
+				VectorStore(Value, &OutValues[i]);
+			}
+
+			// Perform operation where SIMD not possible.
+			for (int32 i = NumToSimd; i < Num; i += 2)
+			{
+				OutValues[i] = Scale * InValues[i];
+				OutValues[i + 1] = -Scale * InValues[i + 1];
+			}
+		}
+	}
+
 	FContiguousSparse2DKernelTransform::FContiguousSparse2DKernelTransform(const int32 NumInElements, const int32 NumOutElements)
 	:	NumIn(NumInElements)
 	,	NumOut(NumOutElements)
diff --git a/Engine/Source/Runtime/SignalProcessing/Private/FloatArrayMath.ispc b/Engine/Source/Runtime/SignalProcessing/Private/FloatArrayMath.ispc
--- a/Engine/Source/Runtime/SignalProcessing/Private/FloatArrayMath.ispc
+++ b/Engine/Source/Runtime/SignalProcessing/Private/FloatArrayMath.ispc
@@ -72,6 +72,29 @@ export void ArrayMeanSquared(const uniform float DataPtr[],
 	OutMean = reduce_add(Sum) / (uniform float)Num;
 }
 
+export void ArrayMax(const uniform float InFloatBuffer1[],
+	const uniform float InFloatBuffer2[],
+	uniform float OutputBuffer[], 
+	const uniform int Num)
+{
+	foreach(i = 0 ... Num)
+	{
+		OutputBuffer[i] = max(InFloatBuffer1[i], InFloatBuffer2[i]);
+	}
+}
+
+export uniform float ArrayMaxAbsValue(const uniform float DataPtr[], const uniform int Num)
+{
+	float Max = 0.f;
+
+	foreach(i = 0 ... Num)
+	{
+		Max = max(abs(DataPtr[i]), Max);
+	}
+
+	return reduce_max(Max);
+}
+
 export uniform float ArrayGetMagnitude(const uniform float Buffer[], const uniform int NumSamples)
 {
 	varying float Sum = 0.0f;
@@ -305,6 +328,36 @@ export void ArrayComplexMultiplyInPlace(const uniform float InData1[],
 	}
 }
 
+export void ArrayComplexMultiplyAdd(const uniform float InA[],
+	const uniform float InB[],
+	uniform float Out[],
+	const uniform int Num)
+{
+#if TARGET_WIDTH == 4
+	static const varying float RealSignFlip = { -1.f, 1.f, -1.f, 1.f };
+#elif TARGET_WIDTH == 8
+	static const varying float RealSignFlip = { -1.f, 1.f, -1.f, 1.f, -1.f, 1.f, -1.f, 1.f };
+#elif TARGET_WIDTH == 16
+	static const varying float RealSignFlip = { -1.f, 1.f, -1.f, 1.f, -1.f, 1.f, -1.f, 1.f, -1.f, 1.f, -1.f, 1.f, -1.f, 1.f, -1.f, 1.f };
+#endif
+
+	foreach(i = 0 ... Num)
+	{
+		const varying float VectorData1 = InA[i];
+		const varying float VectorData2 = InB[i];
+
+		const varying float VectorData1Real = VectorSwizzle(VectorData1, 0, 0, 2, 2);
+		const varying float VectorData1Imag = VectorSwizzle(VectorData1, 1, 1, 3, 3);
+		const varying float VectorData2Swizzle = VectorSwizzle(VectorData2, 1, 0, 3, 2);
+
+		varying float Result = VectorData1Imag * VectorData2Swizzle;
+		Result = Result * RealSignFlip;
+		Result = (VectorData1Real * VectorData2) + Result;
+
+		Out[i] += Result;
+	}
+}
+
 export void ArrayMultiplyByConstant(const uniform float InFloatBuffer[],
 									const uniform float InValue,
 									uniform float OutFloatBuffer[],
@@ -693,3 +746,96 @@ export void TransformArrayRow(const uniform float OffsetInData[],
 
 	OutArray[RowIndex] += reduce_add(RowSum);
 }
+
+export void ArrayFloatToPcm16(const uniform float InBuffer[],
+	uniform int16 OutBuffer[],
+	const uniform int Num)
+{
+	const float MaxSignedShortFloat = (float)0x7FFF;
+
+	foreach(i = 0 ... Num)
+	{
+		OutBuffer[i] = (int16)(InBuffer[i] * MaxSignedShortFloat);
+	}
+}
+
+export void ArrayPcm16ToFloat(const uniform int16 InBuffer[],
+	uniform float OutBuffer[],
+	const uniform int Num)
+{
+	const float InvMaxSignedShortFloat = 1.f / ((float)0x7FFF);
+
+	foreach(i = 0 ... Num)
+	{
+		OutBuffer[i] = ((float)InBuffer[i]) * InvMaxSignedShortFloat;
+	}
+}
+
+export void ArrayAPFLongDelayProcess(const uniform float InSamples[],
+	const uniform float InDelaySamples[],
+	const uniform int InNum,
+	uniform float OutSamples[],
+	uniform float OutDelaySamples[],
+	const uniform float Gain)
+{
+	foreach(i = 0 ... InNum)
+	{
+		OutDelaySamples[i] = (InDelaySamples[i] * Gain) + InSamples[i];
+		OutSamples[i] = (OutDelaySamples[i] * (-Gain)) + InDelaySamples[i];
+	}
+}
+
+export void ArrayLerpFractionalDelay(const uniform float InSamples[],
+	const uniform float InDelays[],
+	const uniform float DelayData[],
+	const uniform int IntegerDelays[],
+	const uniform int InNum,
+	uniform float OutSamples[],
+	const uniform float MaxDelay)
+{
+	foreach(i = 0 ... InNum)
+	{
+		float VFractionalDelays = InDelays[i];
+		VFractionalDelays = max(VFractionalDelays, 0.f);
+		VFractionalDelays = min(VFractionalDelays, MaxDelay);
+
+		float VFloorDelays = floor(VFractionalDelays);
+
+		float VUpperCoefficients = VFractionalDelays - VFloorDelays;
+		float VLowerCoefficients = 1.f - VUpperCoefficients;
+
+		int VIntegerDelays = (int)VFloorDelays;
+		const int VIntegerDelayOffset = IntegerDelays[i];
+		VIntegerDelays = VIntegerDelayOffset - VIntegerDelays;
+
+		#pragma ignore warning(perf)
+		float VLowerSamples = DelayData[VIntegerDelays + 1];
+		#pragma ignore warning(perf)
+		float VUpperSamples = DelayData[VIntegerDelays];
+
+		OutSamples[i] = (VLowerSamples * VLowerCoefficients) + (VUpperSamples * VUpperCoefficients);
+	}
+}
+
+export void ArrayScaledComplexConjugate(const uniform float InValues[],
+	const uniform int Num,
+	uniform float OutValues[],
+	const uniform float Scale)
+{
+	float SignFlip = -Scale;
+
+#if TARGET_WIDTH == 4
+	bool Mask = { true, false, true, false };
+#elif TARGET_WIDTH == 8
+	bool Mask = { true, false, true, false, true, false, true, false };
+#elif TARGET_WIDTH == 16
+	bool Mask = { true, false, true, false, true, false, true, false, true, false, true, false, true, false, true, false };
+#endif
+
+	SignFlip = select(Mask, Scale, SignFlip);
+
+	foreach(i = 0 ... Num)
+	{
+		OutValues[i] = InValues[i] * SignFlip;
+	}
+}
diff --git a/Engine/Source/Runtime/SignalProcessing/Private/LinearInterpFractionalDelay.cpp b/Engine/Source/Runtime/SignalProcessing/Private/LinearInterpFractionalDelay.cpp
--- a/Engine/Source/Runtime/SignalProcessing/Private/LinearInterpFractionalDelay.cpp
+++ b/Engine/Source/Runtime/SignalProcessing/Private/LinearInterpFractionalDelay.cpp
@@ -3,6 +3,7 @@
 #include "DSP/LinearInterpFractionalDelay.h"
 #include "DSP/Dsp.h"
 #include "DSP/BufferVectorOperations.h"
+#include "DSP/FloatArrayMath.h"
 
 namespace Audio
 {
@@ -106,52 +107,7 @@ namespace Audio
 		const float* DelayData = DelayLine->InspectSamples(InNum + MaxDelay + 1);
 		const int32* IntegerDelayOffsetData = IntegerDelayOffsets.GetData();
 
-		const VectorRegister4Float VMaxDelay = MakeVectorRegister((float)MaxDelay, (float)MaxDelay, (float)MaxDelay, (float)MaxDelay);
-		for (int32 i = 0; i < InNum; i += 4)
-		{
-			
-			VectorRegister4Float VFractionalDelays = VectorLoad(&InDelays[i]);
-			// Ensure fractional delays are positive
-			VFractionalDelays = VectorMax(VFractionalDelays, GlobalVectorConstants::FloatZero);
-			VFractionalDelays = VectorMin(VFractionalDelays, VMaxDelay);
-
-			// Separate integer from fraction
-			VectorRegister4Float VFloorDelays = VectorFloor(VFractionalDelays);
-
-			// Determine linear weights
-			VectorRegister4Float VUpperCoefficients = VectorSubtract(VFractionalDelays, VFloorDelays);
-			VectorRegister4Float VLowerCoefficients = VectorSubtract(GlobalVectorConstants::FloatOne, VUpperCoefficients);
-
-
-			// Make integer locations relative to block
-			VectorRegister4Int VIntegerDelays = VectorFloatToInt(VFloorDelays);
-			VectorRegister4Int VIntegerDelayOffset = VectorIntLoadAligned(&IntegerDelayOffsetData[i]);
-			VIntegerDelays = VectorIntSubtract(VIntegerDelayOffset, VIntegerDelays);
-
-			// Lookup samples for interpolation
-			VectorIntStoreAligned(VIntegerDelays, UpperDelayPos);
-			VectorIntStoreAligned(VectorIntAdd(VIntegerDelays, GlobalVectorConstants::IntOne), LowerDelayPos);
-			
-			VectorRegister4Float VLowerSamples = MakeVectorRegister(
-				DelayData[LowerDelayPos[0]],
-				DelayData[LowerDelayPos[1]],
-				DelayData[LowerDelayPos[2]],
-				DelayData[LowerDelayPos[3]]
-			);
-			VectorRegister4Float VUpperSamples = MakeVectorRegister(
-				DelayData[UpperDelayPos[0]],
-				DelayData[UpperDelayPos[1]],
-				DelayData[UpperDelayPos[2]],
-				DelayData[UpperDelayPos[3]]
-			);
-
-			// Interpolate samples
-			VectorRegister4Float VOut = VectorMultiplyAdd(
-				VLowerSamples,
-				VLowerCoefficients,
-				VectorMultiply(VUpperSamples, VUpperCoefficients));
-			VectorStore(VOut, &OutSamples[i]);
-		}
+		ArrayLerpFractionalDelay(InSamples, InDelays, DelayData, IntegerDelayOffsetData, UpperDelayPos, LowerDelayPos, InNum, OutSamples, (float)MaxDelay);
 
 		// Remove unneeded delay line.
 		DelayLine->RemoveSamples(InNum);
diff --git a/Engine/Source/Runtime/SignalProcessing/Private/LongDelayAPF.cpp b/Engine/Source/Runtime/SignalProcessing/Private/LongDelayAPF.cpp
--- a/Engine/Source/Runtime/SignalProcessing/Private/LongDelayAPF.cpp
+++ b/Engine/Source/Runtime/SignalProcessing/Private/LongDelayAPF.cpp
@@ -3,6 +3,7 @@
 #include "DSP/LongDelayAPF.h"
 #include "DSP/BufferVectorOperations.h"
 #include "DSP/Dsp.h"
+#include "DSP/FloatArrayMath.h"
 
 namespace Audio
 {
@@ -128,31 +129,7 @@ namespace Audio
 
 	void FLongDelayAPF::ProcessAudioBlock(const float* InSamples, const float* InDelaySamples, const int32 InNum, float* OutSamples, float* OutDelaySamples)
 	{
-		// Calculate new delay line samples. "w[n] = x[n] + gw[n - d]"
-		int32 NumToSIMD = InNum - (InNum % AUDIO_NUM_FLOATS_PER_VECTOR_REGISTER);
-
-		VectorRegister4Float VG = MakeVectorRegisterFloat(G, G, G, G);
-		VectorRegister4Float VNG = MakeVectorRegisterFloat(-G, -G, -G, -G);
-
-		for (int32 i = 0; i < InNum; i += AUDIO_NUM_FLOATS_PER_VECTOR_REGISTER)
-		{
-			VectorRegister4Float VInDelay = VectorLoadAligned(&InDelaySamples[i]);
-			VectorRegister4Float VInSamples = VectorLoadAligned(&InSamples[i]);
-			// w[n] = x[n] + G * w[n - D]
-			VectorRegister4Float VOutDelay = VectorMultiplyAdd(VInDelay, VG, VInSamples);
-			VectorStoreAligned(VOutDelay, &OutDelaySamples[i]);
-
-			// y[n] = -G * w[n] + w[n - D]
-			VectorRegister4Float VOut = VectorMultiplyAdd(VOutDelay, VNG, VInDelay);
-			VectorStoreAligned(VOut, &OutSamples[i]);
-		}
-
-		// Calculate allpass for remaining samples that we couldn't SIMD
-		for (int32 i = NumToSIMD; i < InNum; i++)
-		{
-			OutDelaySamples[i] = InDelaySamples[i] * G + InSamples[i];
-			OutSamples[i] = OutDelaySamples[i] * -G + InDelaySamples[i];
-		}
+		ArrayAPFLongDelayProcess(InSamples, InDelaySamples, InNum, OutSamples, OutDelaySamples, G);
 	}
 
 	void FLongDelayAPF::Reset() 
diff --git a/Engine/Source/Runtime/SignalProcessing/Private/UniformPartitionConvolution.cpp b/Engine/Source/Runtime/SignalProcessing/Private/UniformPartitionConvolution.cpp
--- a/Engine/Source/Runtime/SignalProcessing/Private/UniformPartitionConvolution.cpp
+++ b/Engine/Source/Runtime/SignalProcessing/Private/UniformPartitionConvolution.cpp
@@ -6,6 +6,7 @@
 #include "DSP/BufferVectorOperations.h"
 #include "DSP/ConvolutionAlgorithm.h"
 #include "DSP/FFTAlgorithm.h"
+#include "DSP/FloatArrayMath.h"
 
 ///////////////////////////////////////////////////////////////////////////////
 // Note on OVERLAP_SAVE
@@ -251,70 +252,7 @@ namespace Audio
 
 	void FUniformPartitionConvolution::VectorComplexMultiplyAdd(const FAlignedFloatBuffer& InA, const FAlignedFloatBuffer& InB, FAlignedFloatBuffer& Out) const
 	{
-		check(InA.Num() == InB.Num());
-		check(Out.Num() == InA.Num());
-
-		const int32 Num = InA.Num();
-		const int32 NumSimd = Num & NumSimdMask;
-
-		// Complex numbers are stored as [real_0, complex_0, real_1, complex_1, ... real_N, complex_N]
-		// So we final amount must be evenly divisble by 2.
-		check(NumSimd % 2 == 0);
-
-		const float* InAData = InA.GetData();
-		const float* InBData = InB.GetData();
-		float* OutData = Out.GetData();
-
-		const VectorRegister4Float SignFlip = MakeVectorRegisterFloat(-1.f, 1.f, -1.f, 1.f);
-
-		for (int32 i = 0; i < NumSimd; i += 4)
-		{
-			// Complex multiply add
-			// Nr = real component of Nth number
-			// Ni = imaginary component of Nth number
-			//
-			//
-			// The input is then 
-			// A1r A1i A2r A2i
-			// B1r B1i B2r B2i
-
-			// VectorA = A1r A1i A2r A2i
-			VectorRegister4Float VectorInA = VectorLoad(&InAData[i]);
-			// Temp12 = A1i A1r A2i A2r
-			VectorRegister4Float Temp1 = VectorSwizzle(VectorInA, 1, 0, 3, 2);
-
-			// VectorB = B1r B1i B2r B2i
-			VectorRegister4Float VectorInB = VectorLoad(&InBData[i]);
-			// Temp2 = B1r B1r B2r B2r
-			VectorRegister4Float Temp2 = VectorSwizzle(VectorInB, 0, 0, 2, 2);
-			// Temp3 = B1i B1i B2i B2i
-			VectorRegister4Float Temp3 = VectorSwizzle(VectorInB, 1, 1, 3, 3);
-
-
-			// VectorA = A1rB1r, A1iB1r, A2rB2r, A2iB2r
-			VectorInA = VectorMultiply(VectorInA, Temp2);
-
-			// Temp1 = A1iB1i, A1rB1i, A2iB2i, A2rb2i
-			Temp1 = VectorMultiply(Temp1, Temp3);
-
-			// Temp1 = -A1iB1i, A1rB1i, -A2iB2i, A2rb2i
-			// Temp1 = A1rB1r - A1iB1i, A1iB1r + A1rB1i, A2rB2r - A2iB2i, A2iB2r + A2rB2i
-			Temp1 = VectorMultiplyAdd(Temp1, SignFlip, VectorInA);
-
-			// VectorOut = O1r + A1rB1r - A1iB1i, O1i + A1iB1r + A1rB1i, O2r + A2rB2r - A2iB2i, O2i + A2iB2r + A2rB2i
-			VectorRegister4Float VectorOut = VectorLoad(&OutData[i]);
-			VectorOut = VectorAdd(Temp1, VectorOut);
-
-			VectorStore(VectorOut, &OutData[i]);
-		}
-
-		for (int32 i = NumSimd; i < Num; i += 2)
-		{
-			// Real output
-			OutData[i] += (InAData[i] * InBData[i]) - (InAData[i + 1] * InBData[i + 1]);
-			// Imaginary output
-			OutData[i + 1] += (InAData[i + 1] * InBData[i]) + (InAData[i] * InBData[i + 1]);
-		}
+		ArrayComplexMultiplyAdd(InA, InB, Out);
 	}
 
 	// Multiply aligned buffer by constant gain.
diff --git a/Engine/Source/Runtime/SignalProcessing/Private/VectorFFT.cpp b/Engine/Source/Runtime/SignalProcessing/Private/VectorFFT.cpp
--- a/Engine/Source/Runtime/SignalProcessing/Private/VectorFFT.cpp
+++ b/Engine/Source/Runtime/SignalProcessing/Private/VectorFFT.cpp
@@ -7,7 +7,6 @@
 #include "DSP/FFTAlgorithm.h"
 #include "DSP/FloatArrayMath.h"
 
-
 namespace Audio
 {
 	// Implementation of a complex FFT.
@@ -205,7 +204,6 @@ namespace Audio
 				int32 WriteIndex;
 			};
 
-
 			// Perform complex conjugate as well as scale.
 			//
 			// @param InValues - Array of floats representing complex values in interleave format.
@@ -214,29 +212,7 @@ namespace Audio
 			// @param Num - Number of floats in array (NOT number of complex values).
 			void ScaledComplexConjugate(const float* InValues, float Scale, float* OutValues, int32 Num)
 			{
-				// Use mask to quickly find out number of values that can be SIMD'd
-				const int32 SIMD_MASK = 0xFFFFFFFC;
-				const int32 NumToSimd = SIMD_MASK & Num;
-
-				// Complex values in a vector are [real_1, imag_1, real_2, imag_2].  
-				// By multipling this value, we flip the sign of the imaginary components, which
-				// is the equivalent of a complex conjugate.
-				const VectorRegister4Float SignFlipImag = MakeVectorRegisterFloat(Scale, -Scale, Scale, -Scale);
-
-				// Perform operation using SIMD
-				for (int32 i = 0; i < NumToSimd; i += 4)
-				{
-					VectorRegister4Float Value = VectorLoad(&InValues[i]);
-					Value = VectorMultiply(SignFlipImag, Value);
-					VectorStore(Value, &OutValues[i]);
-				}
-
-				// Perform operation where SIMD not possible.
-				for (int32 i = NumToSimd; i < Num; i+= 2)
-				{
-					OutValues[i] = Scale * InValues[i];
-					OutValues[i + 1] = -Scale * InValues[i + 1];
-				}
+				ArrayScaledComplexConjugate(InValues, Num, OutValues, Scale);
 			}
 
 			// Perform a radix-4 butterfly which uses constant weights.
@@ -1157,25 +1133,24 @@ namespace Audio
 		const float* BetaImagForwardData = ForwardConvBuffers.BetaImag.GetData();
 
 		// Handle special case of this math to account for cyclical index math.
-		OutComplex[0] = (WorkData[0] * AlphaRealForwardData[0]) 
-			+ (WorkData[1] * AlphaImagForwardData[0]) 
-			+ (WorkData[0] * BetaRealForwardData[0]) 
+		OutComplex[0] = (WorkData[0] * AlphaRealForwardData[0])
+			+ (WorkData[1] * AlphaImagForwardData[0])
+			+ (WorkData[0] * BetaRealForwardData[0])
 			+ (WorkData[1] * BetaImagForwardData[0]);
 
-		OutComplex[1] = (WorkData[1] * AlphaRealForwardData[1]) 
-			+ (WorkData[0] * AlphaImagForwardData[1]) 
-			+ (WorkData[0] * BetaImagForwardData[1]) 
+		OutComplex[1] = (WorkData[1] * AlphaRealForwardData[1])
+			+ (WorkData[0] * AlphaImagForwardData[1])
+			+ (WorkData[0] * BetaImagForwardData[1])
 			+ (WorkData[1] * BetaRealForwardData[1]);
 
-		OutComplex[2] = (WorkData[2] * AlphaRealForwardData[2]) 
-			+ (WorkData[3] * AlphaImagForwardData[2]) 
-			+ (WorkData[FFTSize - 2] * BetaRealForwardData[2]) 
+		OutComplex[2] = (WorkData[2] * AlphaRealForwardData[2])
+			+ (WorkData[3] * AlphaImagForwardData[2])
+			+ (WorkData[FFTSize - 2] * BetaRealForwardData[2])
 			+ (WorkData[FFTSize - 1] * BetaImagForwardData[2]);
-		
 
-		OutComplex[3] = (WorkData[3] * AlphaRealForwardData[3]) 
-			+ (WorkData[2] * AlphaImagForwardData[3]) 
-			+ (WorkData[FFTSize - 2] * BetaImagForwardData[3]) 
+		OutComplex[3] = (WorkData[3] * AlphaRealForwardData[3])
+			+ (WorkData[2] * AlphaImagForwardData[3])
+			+ (WorkData[FFTSize - 2] * BetaImagForwardData[3])
 			+ (WorkData[FFTSize - 1] * BetaRealForwardData[3]);
 
 		// Convert all other values using optimized SIMD
diff --git a/Engine/Source/Runtime/SignalProcessing/Public/DSP/FloatArrayMath.h b/Engine/Source/Runtime/SignalProcessing/Public/DSP/FloatArrayMath.h
--- a/Engine/Source/Runtime/SignalProcessing/Public/DSP/FloatArrayMath.h
+++ b/Engine/Source/Runtime/SignalProcessing/Public/DSP/FloatArrayMath.h
@@ -128,6 +128,12 @@ namespace Audio
 	 */
 	SIGNALPROCESSING_API void ArrayComplexMultiplyInPlace(TArrayView<const float> InValues1, TArrayView<float> InValues2);
 
+	/** Multiplies two complex valued arrays element-wise.
+	 * This assumes elements are in interleaved format [real_0, imag_0, ..., real_N, imag_N]
+	 * Adds result to OutArray
+	 */
+	SIGNALPROCESSING_API void ArrayComplexMultiplyAdd(TArrayView<const float> InValues1, TArrayView<const float> InValues2, TArrayView<float> OutArray);
+
 	/** Multiplies the input float buffer with the given value. */
 	SIGNALPROCESSING_API void ArrayMultiplyByConstant(TArrayView<const float> InFloatBuffer, float InValue, TArrayView<float> OutFloatBuffer);
 
@@ -283,6 +289,7 @@ namespace Audio
 	SIGNALPROCESSING_API void ArrayDoubleSwapBytes(TArrayView<double> InView);
 
 	/** Returns true if host has little endian byte ordering */
+	UE_DEPRECATED(5.5, "Big Endian platforms are no longer supported.")
 	SIGNALPROCESSING_API constexpr bool IsHostLittleEndian()
 	{
 #if PLATFORM_LITTLE_ENDIAN
@@ -292,6 +299,15 @@ namespace Audio
 #endif // PLATFORM_LITTLE_ENDIAN
 	}
 
+	/** All Pass Filter with a long delay. */
+	SIGNALPROCESSING_API void ArrayAPFLongDelayProcess(const float* InSamples, const float* InDelaySamples, const int32 InNum, float* OutSamples, float* OutDelaySamples, const float Gain);
+
+	/** Fractional delay using linear interpolation. */
+	SIGNALPROCESSING_API void ArrayLerpFractionalDelay(const float* InSamples, const float* InDelays, const float* DelayData, const int* IntegerDelays, int* UpperDelayPos, int* LowerDelayPos, const int32 InNum, float* OutSamples, const float MaxDelay);
+
+	/** Perform complex conjugate as well as scale. */
+	SIGNALPROCESSING_API void ArrayScaledComplexConjugate(const float* RESTRICT InValues, const int32 Num, float* RESTRICT OutValues, const float Scale);
+
 	/** FContiguousSparse2DKernelTransform
 	 *
 	 *  FContiguousSparse2DKernelTransform applies a matrix transformation to an input array. 
