Fix for missing data at the end of compressed Ogg Vorbis audio files.

diff --git a/Engine/Source/Runtime/SignalProcessing/Private/OggVorbisEncoder.cpp b/Engine/Source/Runtime/SignalProcessing/Private/OggVorbisEncoder.cpp
--- a/Engine/Source/Runtime/SignalProcessing/Private/OggVorbisEncoder.cpp
+++ b/Engine/Source/Runtime/SignalProcessing/Private/OggVorbisEncoder.cpp
@@ -198,6 +198,9 @@ bool FOggVorbisEncoder::EndFile(TArray<uint8>& OutBytes)
 {
 	check(PrivateState != nullptr);
 
+	// Finish handing over any compressed data.
+	PrivateState->FlushPages(OutBytes);
+
 	// Finalize the DSP analyzer:
 	vorbis_analysis_wrote(&PrivateState->DspState, 0);
 
