Fix unit test failure

#rnx

diff --git a/Engine/Source/Runtime/SignalProcessing/Private/FloatArrayMath.cpp b/Engine/Source/Runtime/SignalProcessing/Private/FloatArrayMath.cpp
--- a/Engine/Source/Runtime/SignalProcessing/Private/FloatArrayMath.cpp
+++ b/Engine/Source/Runtime/SignalProcessing/Private/FloatArrayMath.cpp
@@ -2586,12 +2586,18 @@ namespace Audio
         
 		for (; InIndex < Num; InIndex++)
 		{
-			const uint32 ConvertedValue = uint32(InputPtr[InIndex] * ConversionValue);
+			// Cast to signed integer first because casting a negative float directly to an 
+			// unsigned int is undefined behavior. Some compilers will assign zero. Others
+			// will implicily cast to signed int first and then to unsigned.
+			// https://en.cppreference.com/w/c/language/conversion
+			const int32 ConvertedValue = InputPtr[InIndex] * ConversionValue;
+			const uint32 UnsignedValue = uint32(ConvertedValue);
+
 			uint8* UnsignedOutPtr = (uint8*)&OutPtr[OutIndex];
 
-			UnsignedOutPtr[0] = ConvertedValue & 0xFF;
-			UnsignedOutPtr[1] = ConvertedValue >> 8 & 0xFF;
-			UnsignedOutPtr[2] = ConvertedValue >> 16 & 0xFF;
+			UnsignedOutPtr[0] = UnsignedValue & 0xFF;
+			UnsignedOutPtr[1] = UnsignedValue >> 8 & 0xFF;
+			UnsignedOutPtr[2] = UnsignedValue >> 16 & 0xFF;
 
 			OutIndex += SizeofPCM24;
 		}
