Unshelved from pending changelist '36280913':

Fix array OOB write in AudioFFT

#tests ran the game on consoles

diff --git a/Engine/Source/Runtime/SignalProcessing/Private/AudioFFT.cpp b/Engine/Source/Runtime/SignalProcessing/Private/AudioFFT.cpp
--- a/Engine/Source/Runtime/SignalProcessing/Private/AudioFFT.cpp
+++ b/Engine/Source/Runtime/SignalProcessing/Private/AudioFFT.cpp
@@ -64,7 +64,7 @@ namespace Audio
 		float Phase = 0.0f;
 
 		// Generate the first half of the window:
-		for (int32 FrameIndex = 0; FrameIndex <= Midpoint; FrameIndex++)
+		for (int32 FrameIndex = 0; FrameIndex <= Midpoint && FrameIndex < NumFrames; FrameIndex++)
 		{
 			const float Value = 0.42f - 0.5 * FMath::Cos(Phase) + 0.08 * FMath::Cos(2 * Phase);
 			Phase += PhaseDelta;
