Delete deprecated methods in audio fft to fix static analysis error
#rnx

diff --git a/Engine/Source/Runtime/SignalProcessing/Private/AudioFFT.cpp b/Engine/Source/Runtime/SignalProcessing/Private/AudioFFT.cpp
--- a/Engine/Source/Runtime/SignalProcessing/Private/AudioFFT.cpp
+++ b/Engine/Source/Runtime/SignalProcessing/Private/AudioFFT.cpp
@@ -199,996 +199,8 @@ namespace Audio
 				}
 			}
 		}
-
-		 uint32 NextPowerOf2(uint32 Input)
-		{
-			uint32 Value = 2;
-			while (Value < Input)
-				Value *= 2;
-			return Value;
-		}
-
-		// Fast bit reversal helper function. Can be used if N is a power of 2. Not well exercised.
-		uint32 FastBitReversal(uint32 X, uint32 N)
-		{
-			// Num bits:
-			uint32 NBit = N; // FMath::Log2(FFTSize);
-
-			uint32 Mask = ~0;
-
-			while ((NBit >>= 1) > 0)
-			{
-				Mask ^= (Mask << NBit);
-				X = ((X >> (32u - N + NBit)) & Mask) | ((X << NBit) & ~Mask);
-			}
-
-			return X;
-		}
-
-		// Slow bit reversal helper function. performs bit reversal on an index, bit by bit. N is the number of bits (Log2(FFTSize))
-		uint32 SlowBitReversal(uint32 X, uint32 N)
-		{
-			int32 ReversedX = X;
-			int32 Count = N - 1;
-
-			X >>= 1;
-			while (X > 0)
-			{
-				ReversedX = (ReversedX << 1) | (X & 1);
-				Count--;
-				X >>= 1;
-			}
-
-			return ((ReversedX << Count) & ((1 << N) - 1));
-		}
-
-		// Alternate method for SlowBitReversal. Faster when N >= 7.
-		uint32 SlowBitReversal2(uint32 X, uint32 N)
-		{
-			X = ReverseBits(X);
-			return X >> (32 - N);
-		}
-
-		void ComplexMultiply(const float AReal, const float AImag, const float BReal, const float BImag, float& OutReal, float& OutImag)
-		{
-			OutReal = AReal * BReal - AImag * BImag;
-			OutImag = AReal * BImag + AImag * BReal;
-		}
-
-		// Given:
-		// X = A + iB
-		// Y = C + iD
-		// This function performs the following:
-		// Y = (A*C - B*D) + i (A*D + B*C)
-		void ComplexMultiplyInPlace(const FrequencyBuffer& InFreqBuffer, FrequencyBuffer& OutFreqBuffer)
-		{
-			checkSlow(InFreqBuffer.Real.Num() % 4 == 0);
-			checkSlow(InFreqBuffer.Real.Num() == OutFreqBuffer.Real.Num());
-			int32 NumIterations = InFreqBuffer.Real.Num() / 4;
-
-			
-			const float* XRealBuffer = InFreqBuffer.Real.GetData();
-			const float* XImagBuffer = InFreqBuffer.Imag.GetData();
-
-			float* YRealBuffer = OutFreqBuffer.Real.GetData();
-			float* YImagBuffer = OutFreqBuffer.Imag.GetData();
-
-			for (int32 Iteration = 0; Iteration < NumIterations; Iteration++)
-			{
-				const VectorRegister4Float A = VectorLoad(&XRealBuffer[Iteration * 4]);
-				const VectorRegister4Float B = VectorLoad(&XImagBuffer[Iteration * 4]);
-
-				const VectorRegister4Float C = VectorLoad(&YRealBuffer[Iteration * 4]);
-				const VectorRegister4Float D = VectorLoad(&YImagBuffer[Iteration * 4]);
-
-				const VectorRegister4Float ResultReal = VectorNegateMultiplyAdd(B, D, VectorMultiply(A, C));
-				const VectorRegister4Float ResultImag = VectorMultiplyAdd(A, D, VectorMultiply(B, C));
-
-				VectorStore(ResultReal, &YRealBuffer[Iteration * 4]);
-				VectorStore(ResultImag, &YImagBuffer[Iteration * 4]);
-			}
-		}
-
-		// Given:
-		// X = A + iB
-		// y = c
-		// This function performs the following:
-		// X = (A*c) + i (B*c)
-		void ComplexMultiplyInPlaceByConstant(FrequencyBuffer& InFreqBuffer, float InReal)
-		{
-			checkSlow(InFreqBuffer.Real.Num() % 4 == 0);
-			
-			ArrayMultiplyByConstantInPlace(InFreqBuffer.Real, InReal);
-			ArrayMultiplyByConstantInPlace(InFreqBuffer.Imag, InReal);
-		}
-
-		// Given:
-		// X = A + iB
-		// y = c + id
-		// This function performs the following:
-		// X = (A*c - B*d) + i (A*d + B*C)
-		void ComplexMultiplyInPlaceByConstant(FrequencyBuffer& InFreqBuffer, float InReal, float InImag)
-		{
-			checkSlow(InFreqBuffer.Real.Num() % 4 == 0);
-			int32 NumIterations = InFreqBuffer.Real.Num() / 4;
-
-
-			float* XRealBuffer = InFreqBuffer.Real.GetData();
-			float* XImagBuffer = InFreqBuffer.Imag.GetData();
-
-			const VectorRegister4Float C = VectorSetFloat1(InReal);
-			const VectorRegister4Float D = VectorSetFloat1(InImag);
-
-			for (int32 Iteration = 0; Iteration < NumIterations; Iteration++)
-			{
-				const VectorRegister4Float A = VectorLoad(&XRealBuffer[Iteration * 4]);
-				const VectorRegister4Float B = VectorLoad(&XImagBuffer[Iteration * 4]);
-
-				const VectorRegister4Float ResultReal = VectorNegateMultiplyAdd(B, D, VectorMultiply(A, C));
-				const VectorRegister4Float ResultImag = VectorMultiplyAdd(A, D, VectorMultiply(B, C));
-
-				VectorStore(ResultReal, &XRealBuffer[Iteration * 4]);
-				VectorStore(ResultImag, &XImagBuffer[Iteration * 4]);
-			}
-		}
-
-		// Given:
-		// X = A + iB
-		// Y = C + iD
-		// This function performs the following:
-		// Z = (A*C - B*(-D)) + i (A*(-D) + B*C)
-		//   = (A*C + B*D) + i (-(A*D) + B*C)
-		//   = (A*C + B*D) + i (B*C - A*D)
-		void ComplexMultiplyByConjugate(const FrequencyBuffer& FirstFreqBuffer, const FrequencyBuffer& SecondFrequencyBuffer, FrequencyBuffer& OutFrequencyBuffer)
-		{
-			checkSlow(FirstFreqBuffer.Real.Num() % 4 == 0);
-			checkSlow(FirstFreqBuffer.Real.Num() == SecondFrequencyBuffer.Real.Num());
-			checkSlow(FirstFreqBuffer.Real.Num() == OutFrequencyBuffer.Real.Num());
-
-			int32 NumIterations = FirstFreqBuffer.Real.Num() / 4;
-
-
-			const float* XRealBuffer = FirstFreqBuffer.Real.GetData();
-			const float* XImagBuffer = FirstFreqBuffer.Imag.GetData();
-
-			const float* YRealBuffer = SecondFrequencyBuffer.Real.GetData();
-			const float* YImagBuffer = SecondFrequencyBuffer.Imag.GetData();
-
-			float* ZRealBuffer = OutFrequencyBuffer.Real.GetData();
-			float* ZImagBuffer = OutFrequencyBuffer.Imag.GetData();
-
-			const VectorRegister4Float NegativeOne = VectorSetFloat1(-1.0f);
-
-			for (int32 Iteration = 0; Iteration < NumIterations; Iteration++)
-			{
-				const VectorRegister4Float A = VectorLoad(&XRealBuffer[Iteration * 4]);
-				const VectorRegister4Float B = VectorLoad(&XImagBuffer[Iteration * 4]);
-
-				const VectorRegister4Float C = VectorLoad(&YRealBuffer[Iteration * 4]);
-				const VectorRegister4Float D = VectorLoad(&YImagBuffer[Iteration * 4]);
-
-				const VectorRegister4Float ResultReal = VectorMultiplyAdd(A, C, VectorMultiply(B, D));
-				const VectorRegister4Float ResultImag = VectorNegateMultiplyAdd(A, D, VectorMultiply(B, C));
-
-				VectorStore(ResultReal, &ZRealBuffer[Iteration * 4]);
-				VectorStore(ResultImag, &ZImagBuffer[Iteration * 4]);
-			}
-		}
-
-		// Separates InBuffer (assumed to be mono here)
-		void SeperateInPlace(float* InBuffer, uint32 NumSamples)
-		{
-			check(FMath::CountBits(NumSamples) == 1)
-			const uint32 NumBits = FMath::CountTrailingZeros(NumSamples);
-
-			
-			for (uint32 Index = 0; Index < NumSamples; Index++)
-			{
-				uint32 SwappedIndex = SlowBitReversal(Index, NumBits);
-				if (Index < SwappedIndex)
-				{
-					Swap(InBuffer[Index], InBuffer[SwappedIndex]);
-				}
-			}
-		}
-
-		void SeparateIntoCopy(float* InBuffer, float* OutBuffer, uint32 NumSamples)
-		{
-			check(FMath::CountBits(NumSamples) == 1)
-			const uint32 NumBits = FMath::CountTrailingZeros(NumSamples);
-
-			for (uint32 Index = 0; Index < NumSamples; Index++)
-			{
-				const uint32 ReversedIndex = SlowBitReversal2(Index, NumBits);
-				OutBuffer[ReversedIndex] = InBuffer[Index];
-			}
-		}
-
-		void ComputeButterfliesInPlace(float* OutReal, float* OutImag, uint32 NumSamples)
-		{
-			check(FMath::CountBits(NumSamples) == 1)
-			const uint32 LogNumSamples = FMath::CountTrailingZeros(NumSamples);
-			
-			for (uint32 S = 1; S <= LogNumSamples; S++)
-			{
-				const uint32 M = (1u << S);
-				const uint32 M2 = M >> 1;
-
-				// Initialize sinusoid.
-				float OmegaReal = 1.0f;
-				float OmegaImag = 0.0f;
-
-				// Initialize W of M:
-				float OmegaMReal = FMath::Cos(PI / M2);
-				float OmegaMImag = -FMath::Sin(PI / M2);
-
-				for (uint32 j = 0; j < M2; j++)
-				{
-					for (uint32 k = j; k < NumSamples; k += M)
-					{
-						// Compute twiddle factor:
-						float TwiddleReal, TwiddleImag;
-						const uint32 TwiddleIndex = k + M2;
-
-						ComplexMultiply(OmegaReal, OmegaImag, OutReal[TwiddleIndex], OutImag[TwiddleIndex], TwiddleReal, TwiddleImag);
-
-						// Swap even and odd indices:
-
-						float TempReal = OutReal[k];
-						float TempImag = OutImag[k];
-
-						OutReal[k] = TempReal + TwiddleReal;
-						OutImag[k] = TempImag + TwiddleImag;
-
-						OutReal[TwiddleIndex] = TempReal - TwiddleReal;
-						OutImag[TwiddleIndex] = TempImag - TwiddleImag;
-					}
-
-					// Increment phase of W:
-					ComplexMultiply(OmegaReal, OmegaImag, OmegaMReal, OmegaMImag, OmegaReal, OmegaImag);
-				}
-			}
-		}
-
-		void ComputeButterfliesInPlace2(float* OutReal, float* OutImag, int32 NumSamples)
-		{
-			for (int32 BitPosition = 2; BitPosition <= NumSamples; BitPosition <<= 1)
-			{
-				for (int32 I = 0; I < NumSamples; I += BitPosition)
-				{
-					for (int32 K = 0; K < (BitPosition / 2); K++)
-					{
-						int32 EvenIndex = I + K;
-						int32 OddIndex = I + K + (BitPosition / 2);
-
-						float EvenReal = OutReal[EvenIndex];
-						float EvenImag = OutImag[EvenIndex];
-
-						float OddReal = OutReal[OddIndex];
-						float OddImag = OutImag[OddIndex];
-
-						float Phase = -2.0f * PI * K / ((float)BitPosition);
-						float TwiddleReal = FMath::Cos(Phase);
-						float TwiddleImag = -FMath::Sin(Phase);
-
-						ComplexMultiply(TwiddleReal, TwiddleImag, OddReal, OddImag, TwiddleReal, TwiddleImag);
-
-						// Swap even and odd indices:
-
-						OutReal[EvenIndex] = EvenReal + TwiddleReal;
-						OutImag[EvenIndex] = EvenImag + TwiddleImag;
-
-						OutReal[OddIndex] = EvenReal - TwiddleReal;
-						OutImag[OddIndex] = EvenImag - TwiddleImag;
-					}
-				}
-			}
-		}
-
-		void PerformIterativeFFT(const FFTTimeDomainData_DEPRECATED& InputParams, FFTFreqDomainData_DEPRECATED& OutputParams)
-		{
-			// Separate even and odd elements into real buffer:
-			SeparateIntoCopy(InputParams.Buffer, OutputParams.OutReal, InputParams.NumSamples);
-
-			//Zero out imaginary buffer since the input signal is not complex:
-			FMemory::Memzero(OutputParams.OutImag, InputParams.NumSamples * sizeof(float));
-
-			// Iterate over and compute butterflies.
-			ComputeButterfliesInPlace(OutputParams.OutReal, OutputParams.OutImag, InputParams.NumSamples);
-		}
-
-		void PerformIterativeIFFT(FFTFreqDomainData_DEPRECATED& InputParams, FFTTimeDomainData_DEPRECATED& OutputParams)
-		{
-			TArrayView<float> OutImagView(InputParams.OutImag, OutputParams.NumSamples);
-			TArrayView<float> OutputBufferView(OutputParams.Buffer, OutputParams.NumSamples);
-
-			SeperateInPlace(InputParams.OutReal, OutputParams.NumSamples);
-			SeperateInPlace(InputParams.OutImag, OutputParams.NumSamples);
-
-			// IFFT can be done by performing a forward FFT on the complex conjugate of a frequency domain signal:
-			ArrayMultiplyByConstantInPlace(OutImagView, -1.0f);
-
-			// Iterate over and compute butterflies.
-			ComputeButterfliesInPlace(InputParams.OutReal, InputParams.OutImag, OutputParams.NumSamples);
-
-#if IFFT_PRESERVE_COMPLEX_COMPONENT
-			for (int32 Index = 0; Index < OutputParams.NumSamples; Index++)
-			{
-				const float Real = InputParams.OutReal[Index];
-				const float Imag = InputParams.OutImag[Index];
-				OutputParams.Buffer[Index] = FMath::Sqrt(Real * Real - Imag * Imag);
-			}
-#else
-			FMemory::Memcpy(OutputParams.Buffer, InputParams.OutReal, OutputParams.NumSamples * sizeof(float));
-			
-			// Personal note: This is a very important step in an inverse FFT.
-			Audio::ArrayMultiplyByConstantInPlace(OutputBufferView, 1.0f / OutputParams.NumSamples);
-#endif
-		}
-
-		void PerformDFT(const FFTTimeDomainData_DEPRECATED& InputParams, FFTFreqDomainData_DEPRECATED& OutputParams)
-		{
-			const float* InputBuffer = InputParams.Buffer;
-			float* OutReal = OutputParams.OutReal;
-			float* OutImag = OutputParams.OutImag;
-
-			float N = InputParams.NumSamples;
-
-			for (int32 FreqIndex = 0; FreqIndex < InputParams.NumSamples; FreqIndex++)
-			{
-				float RealSum = 0.0f;
-				float ImagSum = 0.0f;
-
-				for (int32 TimeIndex = 0; TimeIndex < InputParams.NumSamples; TimeIndex++)
-				{
-					const float Exponent = FreqIndex * TimeIndex * PI * 2 / N;
-					RealSum += InputBuffer[TimeIndex] * FMath::Cos(Exponent);
-					ImagSum -= InputBuffer[TimeIndex] * FMath::Sin(Exponent);
-				}
-
-				OutReal[FreqIndex] = RealSum;
-				OutImag[FreqIndex] = ImagSum;
-			}
-		}
-
-		void PerformIDFT(const FFTFreqDomainData_DEPRECATED& InputParams, FFTTimeDomainData_DEPRECATED& OutputParams)
-		{
-			float* OutputBuffer = OutputParams.Buffer;
-			float* InReal = InputParams.OutReal;
-			float* InImag = InputParams.OutImag;
-
-			float N = OutputParams.NumSamples;
-
-			for (int32 TimeIndex = 0; TimeIndex < OutputParams.NumSamples; TimeIndex++)
-			{
-				float RealSum = 0.0f;
-				float ImagSum = 0.0f;
-
-				for (int32 FreqIndex = 0; FreqIndex < OutputParams.NumSamples; FreqIndex++)
-				{
-					const float Exponent = TimeIndex * FreqIndex * PI * 2 / N;
-					RealSum += InReal[FreqIndex] * FMath::Cos(Exponent) - InImag[FreqIndex] * FMath::Sin(Exponent);
-				}
-
-				OutputBuffer[TimeIndex] = RealSum;
-			}
-		}
-
-		void ComputePowerSpectrumNoScaling(const FFTFreqDomainData_DEPRECATED& InFrequencyData, int32 FFTSize, FAlignedFloatBuffer& OutBuffer)
-		{
-			check((FFTSize % 2) == 0);
-
-			if (FFTSize < 1)
-			{
-				// Can't do anything with a zero sized fft.
-				OutBuffer.Reset(0);
-				return;
-			}
-
-			// Spectrum only calculates values for real positive frequencies. 
-			const int32 NumSpectrumValues = (FFTSize / 2) + 1;
-
-			// Resize output buffer
-			OutBuffer.Reset(NumSpectrumValues);
-			OutBuffer.AddUninitialized(NumSpectrumValues);
-
-			TArrayView<float> OutBufferDataView(OutBuffer.GetData(), NumSpectrumValues);
-			TArrayView<const float> RealDataView(InFrequencyData.OutReal, NumSpectrumValues);
-			TArrayView<const float> ImagDataView(InFrequencyData.OutImag, NumSpectrumValues);
-
-			ArrayComplexToPower(RealDataView, ImagDataView, OutBufferDataView);
-		}
-
 	} // namespace FFTIntrinsic
 
-	namespace AudioFFTDeprecated
-	{
-		void CrossCorrelate_DEPRECATED(FrequencyBuffer& FirstBufferFrequencies, FrequencyBuffer& SecondBufferFrequencies, int32 NumSamples, FrequencyBuffer& OutCorrelation);
-		void CrossCorrelate_DEPRECATED(const float* FirstBuffer, const float* SecondBuffer, int32 NumSamples, int32 FFTSize, FrequencyBuffer& FirstBufferFrequencies, FrequencyBuffer& SecondBufferFrequencies, FrequencyBuffer& OutCorrelation);
-		void CrossCorrelate_DEPRECATED(const float* FirstBuffer, const float* SecondBuffer, int32 NumSamples, int32 FFTSize, FrequencyBuffer& OutCorrelation);
-		void CrossCorrelate_DEPRECATED(const float* FirstBuffer, const float* SecondBuffer, int32 NumSamples, int32 FFTSize, float* OutCorrelation, int32 OutCorrelationSamples);
-		void CrossCorrelate_DEPRECATED(FAlignedFloatBuffer& FirstBuffer, FAlignedFloatBuffer& SecondBuffer, FrequencyBuffer& OutCorrelation, bool bZeroPad);
-		void CrossCorrelate_DEPRECATED(FAlignedFloatBuffer& FirstBuffer, FAlignedFloatBuffer& SecondBuffer, FAlignedFloatBuffer& OutCorrelation, bool bZeroPad);
-
-
-		void PerformFFT_DEPRECATED(const FFTTimeDomainData_DEPRECATED& InputParams, FFTFreqDomainData_DEPRECATED& OutputParams)
-		{
-			int32 FFTMethod = CVarFFTMethod.GetValueOnAnyThread();
-			if (FFTMethod)
-			{
-				FFTIntrinsics::PerformDFT(InputParams, OutputParams);
-			}
-			else
-			{
-				FFTIntrinsics::PerformIterativeFFT(InputParams, OutputParams);
-			}
-		}
-
-		void PerformIFFT_DEPRECATED(FFTFreqDomainData_DEPRECATED& InputParams, FFTTimeDomainData_DEPRECATED& OutputParams)
-		{
-			int32 FFTMethod = CVarFFTMethod.GetValueOnAnyThread();
-			if (FFTMethod)
-			{
-				FFTIntrinsics::PerformIDFT(InputParams, OutputParams);
-			}
-			else
-			{
-				FFTIntrinsics::PerformIterativeIFFT(InputParams, OutputParams);
-			}
-		}
-
-		void ComputeMagnitudeSpectrum_DEPRECATED(const FFTFreqDomainData_DEPRECATED& InFrequencyData, int32 FFTSize, FAlignedFloatBuffer& OutBuffer)
-		{
-			if (FFTSize < 1)
-			{
-				OutBuffer.Reset();
-				return;
-			}
-
-			FFTIntrinsics::ComputePowerSpectrumNoScaling(InFrequencyData, FFTSize, OutBuffer);
-
-			const int32 NumSpectrumValues = OutBuffer.Num();
-			if (NumSpectrumValues < 1)
-			{
-				return;
-			}
-
-			const float FFTScale = 1.f / FMath::Sqrt(static_cast<float>(FFTSize));
-
-			float* OutBufferData = OutBuffer.GetData();
-
-			for (int32 i = 0; i < NumSpectrumValues; i++)
-			{
-				// TODO: Currently no vector sqrt. utilize fmath for now.
-				OutBufferData[i] = FMath::Sqrt(OutBufferData[i]);
-			}
-
-			if (NumSpectrumValues > 1)
-			{
-				TArrayView<float> OutBufferView(OutBufferData, NumSpectrumValues - 1);
-				ArrayMultiplyByConstantInPlace(OutBufferView, FFTScale);
-			}
-			OutBufferData[NumSpectrumValues - 1] *= FFTScale;
-		}
-
-		void ComputePowerSpectrum_DEPRECATED(const FFTFreqDomainData_DEPRECATED& InFrequencyData, int32 FFTSize, FAlignedFloatBuffer& OutBuffer)
-		{
-			if (FFTSize < 1)
-			{
-				OutBuffer.Reset();
-				return;
-			}
-
-			FFTIntrinsics::ComputePowerSpectrumNoScaling(InFrequencyData, FFTSize, OutBuffer);
-
-			const int32 NumSpectrumValues = OutBuffer.Num();
-			if (NumSpectrumValues < 1)
-			{
-				return;
-			}
-
-			const float FFTScale = 1.f / static_cast<float>(FFTSize);
-
-			float* OutBufferData = OutBuffer.GetData();
-
-			if (NumSpectrumValues > 1)
-			{
-				TArrayView<float> OutBufferDataView(OutBufferData, NumSpectrumValues - 1);
-				ArrayMultiplyByConstantInPlace(OutBufferDataView, FFTScale);
-			}
-			
-			OutBufferData[NumSpectrumValues - 1] *= FFTScale;
-		}
-
-		void CrossCorrelate_DEPRECATED(FAlignedFloatBuffer& FirstBuffer, FAlignedFloatBuffer& SecondBuffer, FAlignedFloatBuffer& OutCorrelation, bool bZeroPad /*= true*/)
-		{
-			FrequencyBuffer OutputCorrelationFrequencies;
-			CrossCorrelate_DEPRECATED(FirstBuffer, SecondBuffer, OutputCorrelationFrequencies, bZeroPad);
-
-			OutCorrelation.Reset();
-			OutCorrelation.AddUninitialized(OutputCorrelationFrequencies.Real.Num());
-
-			// Perform IFFT into OutCorrelation:
-			FFTFreqDomainData_DEPRECATED FreqDomainData = 
-			{
-				OutputCorrelationFrequencies.Real.GetData(),
-				OutputCorrelationFrequencies.Imag.GetData()
-			};
-
-			FFTTimeDomainData_DEPRECATED TimeDomainData =
-			{
-				OutCorrelation.GetData(),
-				OutCorrelation.Num()
-			};
-
-			AudioFFTDeprecated::PerformIFFT_DEPRECATED(FreqDomainData, TimeDomainData);
-		}
-
-		void CrossCorrelate_DEPRECATED(FAlignedFloatBuffer& FirstBuffer, FAlignedFloatBuffer& SecondBuffer, FrequencyBuffer& OutCorrelation, bool bZeroPad /*= true*/)
-		{
-			const int32 NumSamples = FMath::Max(FirstBuffer.Num(), SecondBuffer.Num());
-
-			if (!bZeroPad)
-			{
-				int32 FFTLength = FFTIntrinsics::NextPowerOf2(NumSamples - 1);
-
-				FirstBuffer.AddZeroed(FFTLength - FirstBuffer.Num());
-				SecondBuffer.AddZeroed(FFTLength - SecondBuffer.Num());
-			}
-			else
-			{
-				checkSlow(FirstBuffer.Num() == SecondBuffer.Num() && FMath::IsPowerOfTwo(FirstBuffer.Num()));
-			}
-
-			CrossCorrelate_DEPRECATED(FirstBuffer.GetData(), SecondBuffer.GetData(), NumSamples, FirstBuffer.Num(), OutCorrelation);
-		}
-
-		void CrossCorrelate_DEPRECATED(const float* FirstBuffer, const float* SecondBuffer, int32 NumSamples, int32 FFTSize, float* OutCorrelation, int32 OutCorrelationSamples)
-		{
-			FrequencyBuffer OutputCorrelationFrequencies;
-			CrossCorrelate_DEPRECATED(FirstBuffer, SecondBuffer, NumSamples, FFTSize, OutputCorrelationFrequencies);
-
-			checkSlow(FFTSize == OutCorrelationSamples);
-
-			// Perform IFFT into OutCorrelation:
-			FFTFreqDomainData_DEPRECATED FreqDomainData =
-			{
-				OutputCorrelationFrequencies.Real.GetData(),
-				OutputCorrelationFrequencies.Imag.GetData()
-			};
-
-			FFTTimeDomainData_DEPRECATED TimeDomainData =
-			{
-				OutCorrelation,
-				OutCorrelationSamples
-			};
-
-			AudioFFTDeprecated::PerformIFFT_DEPRECATED(FreqDomainData, TimeDomainData);
-		}
-
-		void CrossCorrelate_DEPRECATED(const float* FirstBuffer, const float* SecondBuffer, int32 NumSamples, int32 FFTSize, FrequencyBuffer& OutCorrelation)
-		{
-			FrequencyBuffer FirstBufferFrequencies;
-			FrequencyBuffer SecondBufferFrequencies;
-
-			CrossCorrelate_DEPRECATED(FirstBuffer, SecondBuffer, NumSamples, FFTSize, FirstBufferFrequencies, SecondBufferFrequencies, OutCorrelation);
-		}
-
-		void CrossCorrelate_DEPRECATED(const float* FirstBuffer, const float* SecondBuffer, int32 NumSamples, int32 FFTSize, FrequencyBuffer& FirstBufferFrequencies, FrequencyBuffer& SecondBufferFrequencies, FrequencyBuffer& OutCorrelation)
-		{
-			checkSlow(FMath::IsPowerOfTwo(FFTSize));
-			OutCorrelation.InitZeroed(FFTSize);
-			
-			// Perform FFT on First Buffer input:
-			FirstBufferFrequencies.InitZeroed(FFTSize);
-			FFTTimeDomainData_DEPRECATED TimeDomainData =
-			{
-				const_cast<float*>(FirstBuffer), // Note: Because we use the same struct for the IFFT as we do the FFT, we can't make this property const.
-				FFTSize
-			};
-
-			FFTFreqDomainData_DEPRECATED FreqDomainData =
-			{
-				FirstBufferFrequencies.Real.GetData(),
-				FirstBufferFrequencies.Imag.GetData()
-			};
-			
-			AudioFFTDeprecated::PerformFFT_DEPRECATED(TimeDomainData, FreqDomainData);
-
-			// Perform FFT on second buffer of input:
-			SecondBufferFrequencies.InitZeroed(FFTSize);
-			TimeDomainData =
-			{
-				const_cast<float*>(SecondBuffer),
-				FFTSize
-			};
-
-			FreqDomainData = 
-			{
-				SecondBufferFrequencies.Real.GetData(),
-				SecondBufferFrequencies.Imag.GetData()
-			};
-
-			AudioFFTDeprecated::PerformFFT_DEPRECATED(TimeDomainData, FreqDomainData);
-
-			CrossCorrelate_DEPRECATED(FirstBufferFrequencies, SecondBufferFrequencies, NumSamples, OutCorrelation);
-		}
-
-		void CrossCorrelate_DEPRECATED(FrequencyBuffer& FirstBufferFrequencies, FrequencyBuffer& SecondBufferFrequencies, int32 NumSamples, FrequencyBuffer& OutCorrelation)
-		{
-			FFTIntrinsics::ComplexMultiplyByConjugate(FirstBufferFrequencies, SecondBufferFrequencies, OutCorrelation);
-		}
-
-	} // Namespace AudioFFTDeprecated
-
-	void PerformFFT(const FFTTimeDomainData_DEPRECATED& InputParams, FFTFreqDomainData_DEPRECATED& OutputParams)
-	{
-		AudioFFTDeprecated::PerformFFT_DEPRECATED(InputParams, OutputParams);
-	}
-
-	void PerformIFFT(FFTFreqDomainData_DEPRECATED& InputParams, FFTTimeDomainData_DEPRECATED& OutputParams)
-	{
-		AudioFFTDeprecated::PerformIFFT_DEPRECATED(InputParams, OutputParams);
-	}
-
-	// Implementation of fft algorithm interface
-	class FAudioFFTAlgorithm : public IFFTAlgorithm
-	{
-			const int32 FFTSize;
-			const int32 NumOutputFFTElements;
-			FFTFreqDomainData_DEPRECATED FreqDomainData;
-
-			FAlignedFloatBuffer FreqRealBuffer;
-			FAlignedFloatBuffer FreqImagBuffer;
-			
-
-		public:
-
-			FAudioFFTAlgorithm(int32 InFFTSize)
-			:	FFTSize(InFFTSize)
-			,	NumOutputFFTElements((InFFTSize / 2) + 1)
-			{
-				// For freq domain data, we need separate buffers since we are expecting
-				// interleaved [real, imag] data.
-				FreqRealBuffer.AddUninitialized(FFTSize);
-				FreqImagBuffer.AddUninitialized(FFTSize);
-
-				FreqDomainData.OutReal = FreqRealBuffer.GetData();
-				FreqDomainData.OutImag = FreqImagBuffer.GetData();
-			}
-
-			// virtual destructor for inheritance.
-			virtual ~FAudioFFTAlgorithm() {}
-
-			// Number of elements in FFT
-			virtual int32 Size() const override { return FFTSize; }
-
-			// Scaling applied when performing forward FFT.
-			virtual EFFTScaling ForwardScaling() const { return EFFTScaling::MultipliedBySqrtFFTSize; }
-
-			/* Scaling applied when performing inverse FFT. */
-			virtual EFFTScaling InverseScaling() const { return EFFTScaling::DividedBySqrtFFTSize; }
-
-			
-			// ForwardRealToComplex
-			// InReal - Array of floats to input into fourier transform. Must have FFTSize() elements.
-			// OutComplex - Array of floats to store output of fourier transform. Must have (FFTSize() + 2) float elements which represent ((FFTSize() / 2) + 1) complex numbers in interleaved format.
-			virtual void ForwardRealToComplex(const float* RESTRICT InReal, float* RESTRICT OutComplex) override
-			{
-				const FFTTimeDomainData_DEPRECATED TimeDomainData = {const_cast<float*>(InReal), FFTSize};
-				AudioFFTDeprecated::PerformFFT_DEPRECATED(TimeDomainData, FreqDomainData);
-
-				// Convert FFT output data to interleaved format.
-				int32 OutPos = 0;
-				for (int32 i = 0; i < NumOutputFFTElements; i++)
-				{
-					OutComplex[OutPos] = FreqDomainData.OutReal[i];
-					OutPos++;
-					OutComplex[OutPos] = FreqDomainData.OutImag[i];
-					OutPos++;
-				}
-			}
-
-			// InverseComplexToReal
-			// InComplex - Array of floats to input into inverse fourier transform. Must have (FFTSize() + 2) float elements which represent ((FFTSize() / 2) + 1) complex numbers in interleaved format.
-			// OutReal - Array of floats to store output of inverse fourier transform. Must have FFTSize() elements.
-			virtual void InverseComplexToReal(const float* RESTRICT InComplex, float* RESTRICT OutReal) override
-			{
-
-				// For the complex data the phase must be flipped for negative frequencies (or frequencies above nyquist depending on the way you think about it.)
-				
-				// Copy from 0Hz -> Nyquist
-				int32 ComplexPos = 0;
-				for (int32 i = 0; i < NumOutputFFTElements; i++)
-				{
-					FreqDomainData.OutReal[i] = InComplex[ComplexPos];
-					ComplexPos++;
-					FreqDomainData.OutImag[i] = InComplex[ComplexPos];
-					ComplexPos++;
-				}
-
-				// Perform mirror
-				int32 InFreqPos = NumOutputFFTElements - 2;
-				for (int32 MirrorPos = NumOutputFFTElements; MirrorPos < FFTSize; MirrorPos++)
-				{
-					FreqDomainData.OutReal[MirrorPos] = FreqDomainData.OutReal[InFreqPos];
-					FreqDomainData.OutImag[MirrorPos] = -FreqDomainData.OutImag[InFreqPos];
-					InFreqPos--;
-				}
-
-				FFTTimeDomainData_DEPRECATED TimeDomainData = {OutReal, FFTSize};
-				AudioFFTDeprecated::PerformIFFT_DEPRECATED(FreqDomainData, TimeDomainData);
-			}
-
-			// BatchForwardRealToComplex
-			virtual void BatchForwardRealToComplex(int32 InCount, const float* const RESTRICT InReal[], float* RESTRICT OutComplex[]) override
-			{
-				for (int32 i = 0; i < InCount; i++)
-				{
-					ForwardRealToComplex(InReal[i], OutComplex[i]);
-				}
-			}
-
-			// BatchInverseComplexToReal
-			virtual void BatchInverseComplexToReal(int32 InCount, const float* const RESTRICT InComplex[], float* RESTRICT OutReal[]) override
-			{
-				for (int32 i = 0; i < InCount; i++)
-				{
-					InverseComplexToReal(InComplex[i], OutReal[i]);
-				}
-			}
-	};
-
-
-	// FFT Algorithm factory for this FFT implementation
-	FAudioFFTAlgorithmFactory::~FAudioFFTAlgorithmFactory() {}
-
-	// Name of this fft algorithm factory. 
-	FName FAudioFFTAlgorithmFactory::GetFactoryName() const 
-	{
-		static const FName FactoryName = FName(TEXT("OriginalFFT_Deprecated"));
-		return FactoryName;
-	}
-
-	// If true, this implementation uses hardware acceleration.
-	bool FAudioFFTAlgorithmFactory::IsHardwareAccelerated() const
-	{
-		return false;
-	}
-
-	// If true, this implementation requires input and output arrays to be 128 bit aligned.
-	bool FAudioFFTAlgorithmFactory::Expects128BitAlignedArrays() const
-	{
-		return false;
-	}
-
-	// Returns true if the input settings are supported by this factory.
-	bool FAudioFFTAlgorithmFactory::AreFFTSettingsSupported(const FFFTSettings& InSettings) const 
-	{
-		return (InSettings.Log2Size > 1) && (InSettings.Log2Size < 30);
-	}
-
-	// Create a new FFT algorithm.
-	TUniquePtr<IFFTAlgorithm> FAudioFFTAlgorithmFactory::NewFFTAlgorithm(const FFFTSettings& InSettings)
-	{
-		check(AreFFTSettingsSupported(InSettings));
-
-		// equivalent of 2^(InSettings.Log2Size)
-		int32 FFTSize = 1 << InSettings.Log2Size;
-
-		return MakeUnique<FAudioFFTAlgorithm>(FFTSize);
-	}
-
-	void ComputePowerSpectrum(const FFTFreqDomainData_DEPRECATED& InFrequencyData, int32 FFTSize, FAlignedFloatBuffer& OutBuffer)
-	{
-		AudioFFTDeprecated::ComputePowerSpectrum_DEPRECATED(InFrequencyData, FFTSize, OutBuffer);
-	}
-
-	void ComputeMagnitudeSpectrum(const FFTFreqDomainData_DEPRECATED& InFrequencyData, int32 FFTSize, FAlignedFloatBuffer& OutBuffer)
-	{
-		AudioFFTDeprecated::ComputeMagnitudeSpectrum_DEPRECATED(InFrequencyData, FFTSize, OutBuffer);
-	}
-
-	void ComputeSpectrum(ESpectrumType InSpectrumType, const FFTFreqDomainData_DEPRECATED& InFrequencyData, int32 FFTSize, FAlignedFloatBuffer& OutBuffer)
-	{
-		switch (InSpectrumType)
-		{
-			case ESpectrumType::MagnitudeSpectrum:
-				AudioFFTDeprecated::ComputeMagnitudeSpectrum_DEPRECATED(InFrequencyData, FFTSize, OutBuffer);
-				break;
-
-			case ESpectrumType::PowerSpectrum:
-				AudioFFTDeprecated::ComputePowerSpectrum_DEPRECATED(InFrequencyData, FFTSize, OutBuffer);
-				break;
-
-			default:
-				checkf(false, TEXT("Unhandled Audio::ESpectrumType"));
-		}
-	}
-
-	void CrossCorrelate(FAlignedFloatBuffer& FirstBuffer, FAlignedFloatBuffer& SecondBuffer, FAlignedFloatBuffer& OutCorrelation, bool bZeroPad /*= true*/)
-	{
-		AudioFFTDeprecated::CrossCorrelate_DEPRECATED(FirstBuffer, SecondBuffer, OutCorrelation, bZeroPad);
-	}
-
-	void CrossCorrelate(FAlignedFloatBuffer& FirstBuffer, FAlignedFloatBuffer& SecondBuffer, FrequencyBuffer& OutCorrelation, bool bZeroPad /*= true*/)
-	{
-		AudioFFTDeprecated::CrossCorrelate_DEPRECATED(FirstBuffer, SecondBuffer, OutCorrelation, bZeroPad);
-	}
-
-	void CrossCorrelate(const float* FirstBuffer, const float* SecondBuffer, int32 NumSamples, int32 FFTSize, float* OutCorrelation, int32 OutCorrelationSamples)
-	{
-		AudioFFTDeprecated::CrossCorrelate_DEPRECATED(FirstBuffer, SecondBuffer, NumSamples, FFTSize, OutCorrelation, OutCorrelationSamples);
-	}
-
-	void CrossCorrelate(const float* FirstBuffer, const float* SecondBuffer, int32 NumSamples, int32 FFTSize, FrequencyBuffer& OutCorrelation)
-	{
-		AudioFFTDeprecated::CrossCorrelate_DEPRECATED(FirstBuffer, SecondBuffer, NumSamples, FFTSize, OutCorrelation);
-	}
-
-	void CrossCorrelate(const float* FirstBuffer, const float* SecondBuffer, int32 NumSamples, int32 FFTSize, FrequencyBuffer& FirstBufferFrequencies, FrequencyBuffer& SecondBufferFrequencies, FrequencyBuffer& OutCorrelation)
-	{
-		AudioFFTDeprecated::CrossCorrelate_DEPRECATED(FirstBuffer, SecondBuffer, NumSamples, FFTSize, FirstBufferFrequencies, SecondBufferFrequencies, OutCorrelation);
-	}
-
-	void CrossCorrelate(FrequencyBuffer& FirstBufferFrequencies, FrequencyBuffer& SecondBufferFrequencies, int32 NumSamples, FrequencyBuffer& OutCorrelation)
-	{
-		AudioFFTDeprecated::CrossCorrelate_DEPRECATED(FirstBufferFrequencies, SecondBufferFrequencies, NumSamples, OutCorrelation);
-	}
-
-	FFFTConvolver_DEPRECATED::FFFTConvolver_DEPRECATED()
-		: BlockSize(0)
-	{
-	}
-
-	void FFFTConvolver_DEPRECATED::ConvolveBlock(float* InputAudio, int32 NumSamples)
-	{
-		checkSlow(NumSamples != 0);
-
-		const int32 FFTSize = FilterFrequencies.Real.Num();
-
-		// Zero-pad the input to the FFT size:
-		TimeDomainInputBuffer.Reset(FFTSize);
-		TimeDomainInputBuffer.AddZeroed(FFTSize);
-		FMemory::Memcpy(TimeDomainInputBuffer.GetData(), InputAudio, NumSamples * sizeof(float));
-
-		FFTTimeDomainData_DEPRECATED TimeDomainData = 
-		{ 
-			TimeDomainInputBuffer.GetData(), // Buffer
-			TimeDomainInputBuffer.Num()      // FFTSize
-		};
-
-		FFTFreqDomainData_DEPRECATED FreqDomainData = 
-		{
-			InputFrequencies.Real.GetData(),
-			InputFrequencies.Imag.GetData()
-		};
-
-		AudioFFTDeprecated::PerformFFT_DEPRECATED(TimeDomainData, FreqDomainData);
-		FFTIntrinsics::ComplexMultiplyInPlace(FilterFrequencies, InputFrequencies);
-		AudioFFTDeprecated::PerformIFFT_DEPRECATED(FreqDomainData, TimeDomainData);
-
-		// Copy back from our temporary buffer to InputAudio:
-		FMemory::Memcpy(InputAudio, TimeDomainInputBuffer.GetData(), NumSamples * sizeof(float));
-		
-		// TODO: Currently we don't support COLA buffers being larger than our block size.
-		checkSlow(NumSamples >= COLABuffer.Num());
-		SumInCOLABuffer(InputAudio, COLABuffer.Num());
-		
-		const int32 COLASize = BlockSize - 1;
-		const int32 COLAOffset = NumSamples;
-		SetCOLABuffer(TimeDomainInputBuffer.GetData() + COLAOffset, COLASize);
-	}
-
-	void FFFTConvolver_DEPRECATED::SumInCOLABuffer(float* InputAudio, int32 NumSamples)
-	{
-		TArrayView<const float> COLABufferView(COLABuffer.GetData(), NumSamples);
-		TArrayView<float> InputAudioView(InputAudio, NumSamples);
-
-		ArrayMixIn(COLABufferView, InputAudioView);
-	}
-
-	void FFFTConvolver_DEPRECATED::SetCOLABuffer(float* InAudio, int32 NumSamples)
-	{
-		COLABuffer.SetNumUninitialized(NumSamples, EAllowShrinking::No);
-		FMemory::Memcpy(COLABuffer.GetData(), InAudio, NumSamples * sizeof(float));
-	}
-
-	void FFFTConvolver_DEPRECATED::ProcessAudio(float* InputAudio, int32 NumSamples)
-	{
-		checkSlow(BlockSize != 0);
-
-		const int32 NumFullBlocks = NumSamples / BlockSize;
-		const int32 BlockRemainer = NumSamples % BlockSize;
-
-		for (int32 Index = 0; Index < NumFullBlocks; Index++)
-		{	
-			ConvolveBlock(&InputAudio[BlockSize * Index], BlockSize);
-		}
-
-		if (BlockRemainer)
-		{
-			// Perform Convolve Block on the last buffer:
-			ConvolveBlock(&InputAudio[BlockSize* NumFullBlocks], BlockRemainer);
-		}
-	}
-
-	void FFFTConvolver_DEPRECATED::SetFilter(const float* InWindowReal, const float* InWindowImag, int32 FilterSize, int32 FFTSize)
-	{
-		checkSlow(FilterSize != 0);
-
-		// TODO: Support non-power of two WindowSizes. To do this we'll need to be able to accumulate COLA and add partial COLA buffers to individual blocks.
-		check(FMath::IsPowerOfTwo(FFTSize) && FMath::IsPowerOfTwo(FilterSize) &&  FFTSize >= FilterSize * 2 - 1);
-
-		if (FFTSize != FilterFrequencies.Real.Num())
-		{
-			FilterFrequencies.InitZeroed(FFTSize);
-		}
-
-		FilterFrequencies.CopyFrom(InWindowReal, InWindowImag, FFTSize);
-		BlockSize = FilterSize;
-
-		InputFrequencies.InitZeroed(FFTSize);
-		COLABuffer.Reset(BlockSize);
-		COLABuffer.AddZeroed(BlockSize);
-	}
-
-	void FFFTConvolver_DEPRECATED::SetFilter(const FrequencyBuffer& InFilterFrequencies, int32 FilterSize)
-	{
-		checkSlow(FilterSize != 0);
-
-		const int32 FilterFFTSize = InFilterFrequencies.Real.Num();
-		checkSlow(FMath::IsPowerOfTwo(FilterFFTSize) && FilterFFTSize >= FilterSize * 2 - 1);
-
-		if (FilterFFTSize != FilterFrequencies.Real.Num())
-		{
-			FilterFrequencies.InitZeroed(FilterFFTSize);
-		}
-
-		FilterFrequencies.CopyFrom(InFilterFrequencies);
-		BlockSize = FilterSize;
-
-		InputFrequencies.InitZeroed(FilterFFTSize);
-		COLABuffer.Reset(BlockSize);
-		COLABuffer.AddZeroed(BlockSize);
-	}
-
-	void FFFTConvolver_DEPRECATED::SetFilter(const float* TimeDomainBuffer, int32 FilterSize)
-	{
-		checkSlow(FilterSize != 0);
-
-		int32 FilterFFTSize = FFTIntrinsics::NextPowerOf2(FilterSize * 2 - 1);
-		checkSlow(FMath::IsPowerOfTwo(FilterFFTSize) && FilterFFTSize >= FilterSize * 2 - 1);
-
-		FilterFrequencies.InitZeroed(FilterFFTSize);
-		TimeDomainInputBuffer.Reset(FilterFFTSize);
-		TimeDomainInputBuffer.AddZeroed(FilterFFTSize);
-		FMemory::Memcpy(TimeDomainInputBuffer.GetData(), TimeDomainBuffer, FilterSize * sizeof(float));
-
-		FFTTimeDomainData_DEPRECATED TimeDomainData =
-		{
-			TimeDomainInputBuffer.GetData(),
-			TimeDomainInputBuffer.Num()
-		};
-
-		FFTFreqDomainData_DEPRECATED FreqDomainData =
-		{
-			FilterFrequencies.Real.GetData(),
-			FilterFrequencies.Imag.GetData()
-		};
-
-		AudioFFTDeprecated::PerformFFT_DEPRECATED(TimeDomainData, FreqDomainData);
-		BlockSize = FilterSize;
-		InputFrequencies.InitZeroed(FilterFFTSize);
-		COLABuffer.Reset(BlockSize);
-		COLABuffer.AddZeroed(BlockSize);
-	}
-
-	void FFFTConvolver_DEPRECATED::SetFilter(const FAlignedFloatBuffer& TimeDomainBuffer)
-	{
-		SetFilter(TimeDomainBuffer.GetData(), TimeDomainBuffer.Num());
-	}
-
 	int32 CeilLog2(int32 InNum)
 	{
 		static constexpr int32 MaxValue = 0x40000000;
diff --git a/Engine/Source/Runtime/SignalProcessing/Private/VoiceProcessing.cpp b/Engine/Source/Runtime/SignalProcessing/Private/VoiceProcessing.cpp
--- a/Engine/Source/Runtime/SignalProcessing/Private/VoiceProcessing.cpp
+++ b/Engine/Source/Runtime/SignalProcessing/Private/VoiceProcessing.cpp
@@ -187,193 +187,4 @@ namespace Audio
 		const float UnclampedGain = GainTarget / InAmplitude;
 		return FMath::Clamp(UnclampedGain, GainMin, GainMax);
 	}
-
-	FAdaptiveFilter_DEPRECATED::FAdaptiveFilter_DEPRECATED(int32 FilterLength, int32 AudioCallbackSize)
-		: WindowSize(FilterLength)
-		, CurrentStepsUntilConvergence(0)
-	{
-		const int32 FFTSize = FFTIntrinsics::NextPowerOf2(FilterLength * 2 - 1);
-
-		//Initialize our weights to 0:
-		WeightDeltas.InitZeroed(FFTSize);
-		CurrentWeights.InitZeroed(FFTSize);
-		InputFrequencies.InitZeroed(FFTSize);
-
-		// For passthrough, we initialize our current weights to 1:
-		for (int32 Index = 0; Index < CurrentWeights.Real.Num(); Index++)
-		{
-			CurrentWeights.Real[Index] = 1.0f;
-		}
-
-		Convolver.SetFilter(CurrentWeights, FilterLength);
-	}
-
-	void FAdaptiveFilter_DEPRECATED::AdaptFilter()
-	{
-		// If we've suitably converged, we avoid incrementing our weights.
-		if (CurrentStepsUntilConvergence <= 0)
-		{
-			return;
-		}
-
-		IncrementWeights();
-		CurrentStepsUntilConvergence--;
-		Convolver.SetFilter(CurrentWeights, WindowSize);
-	}
-
-	void FAdaptiveFilter_DEPRECATED::SetWeightDeltas(const float* InWeightsReal, const float* InWeightsImag, int32 NumWeights, float InLearningRate)
-	{
-		// ProcessAudio::NumWeights needs to use the same 
-		checkSlow(NumWeights == WeightDeltas.Real.Num());
-
-		// If NumWeights is not divisible by four, something is terribly wrong, since we are using an FFT with this.
-		checkSlow(NumWeights % 4 == 0);
-		const float* CurrentRealBuffer = CurrentWeights.Real.GetData();
-		const float* CurrentImagBuffer = CurrentWeights.Imag.GetData();
-
-		float* DeltaRealBuffer = WeightDeltas.Real.GetData();
-		float* DeltaImagBuffer = WeightDeltas.Imag.GetData();
-
-		// Our convergence rate is 1 / our NumConvergence steps.
-		const VectorRegister4Float ConvergenceRate = VectorSetFloat1(InLearningRate);
-		CurrentStepsUntilConvergence = FMath::CeilToInt(1.0f / InLearningRate);
-		
-		for (int32 WeightIndex = 0; WeightIndex < NumWeights; WeightIndex+= 4)
-		{
-			VectorRegister4Float TargetReal = VectorLoad(&InWeightsReal[WeightIndex]);
-			VectorRegister4Float CurrentReal = VectorLoad(&CurrentRealBuffer[WeightIndex]);
-			const VectorRegister4Float DeltaReal = VectorMultiply(VectorSubtract(TargetReal, CurrentReal), ConvergenceRate);
-			VectorStore(DeltaReal, &DeltaRealBuffer[WeightIndex]);
-
-			VectorRegister4Float TargetImag = VectorLoad(&InWeightsImag[WeightIndex]);
-			VectorRegister4Float CurrentImag = VectorLoad(&CurrentImagBuffer[WeightIndex]);
-			const VectorRegister4Float DeltaImag = VectorMultiply(VectorSubtract(TargetImag, CurrentImag), ConvergenceRate);
-			VectorStore(DeltaImag, &DeltaImagBuffer[WeightIndex]);
-		}
-	}
-
-	void FAdaptiveFilter_DEPRECATED::IncrementWeights()
-	{
-		const int32 NumWeights = CurrentWeights.Real.Num();
-		float* CurrentRealBuffer = CurrentWeights.Real.GetData();
-		float* CurrentImagBuffer = CurrentWeights.Imag.GetData();
-
-		const float* DeltasRealBuffer = WeightDeltas.Real.GetData();
-		const float* DeltasImagBuffer = WeightDeltas.Imag.GetData();
-
-		for (int32 WeightIndex = 0; WeightIndex < NumWeights; WeightIndex += 4)
-		{
-			VectorRegister4Float DeltaReal = VectorLoad(&DeltasRealBuffer[WeightIndex]);
-			VectorRegister4Float CurrentReal = VectorLoad(&CurrentRealBuffer[WeightIndex]);
-			CurrentReal = VectorAdd(CurrentReal, DeltaReal);
-			VectorStore(CurrentReal, &CurrentRealBuffer[WeightIndex]);
-
-			VectorRegister4Float DeltaImag = VectorLoad(&DeltasImagBuffer[WeightIndex]);
-			VectorRegister4Float CurrentImag = VectorLoad(&CurrentImagBuffer[WeightIndex]);
-			CurrentImag = VectorAdd(CurrentImag, DeltaImag);
-			VectorStore(CurrentImag, &CurrentImagBuffer[WeightIndex]);
-		}
-	}
-
-	void FAdaptiveFilter_DEPRECATED::ProcessAudio(float* InAudio, int32 NumSamples)
-	{
-		AdaptFilter();
-		Convolver.ProcessAudio(InAudio, NumSamples);
-	}
-
-	void FAdaptiveFilter_DEPRECATED::SetWeights(const FrequencyBuffer& InFilterWeights, int32 FilterLength, float InLearningRate)
-	{
-		SetWeightDeltas(InFilterWeights.Real.GetData(), InFilterWeights.Imag.GetData(), InFilterWeights.Real.Num(), InLearningRate);
-	}
-
-	FFDAPFilterComputer_DEPRECATED::FFDAPFilterComputer_DEPRECATED()
-	{
-
-	}
-
-	void FFDAPFilterComputer_DEPRECATED::GenerateWeights(const float* IncomingSignal, int32 NumIncomingSamples, const float* OutgoingSignal, int32 NumOutgoingSamples, FrequencyBuffer& OutWeights)
-	{
-		int32 FFTSize = FFTIntrinsics::NextPowerOf2(NumIncomingSamples + NumOutgoingSamples - 1);
-		
-		// zero pad inputs appropriately:
-		ZeroPaddedIncomingBuffer.Reset();
-		ZeroPaddedIncomingBuffer.AddZeroed(FFTSize);
-		FMemory::Memcpy(ZeroPaddedIncomingBuffer.GetData(), IncomingSignal, NumIncomingSamples * sizeof(float));
-
-		ZeroPaddedOutgoingBuffer.Reset();
-		ZeroPaddedOutgoingBuffer.AddZeroed(FFTSize);
-		FMemory::Memcpy(ZeroPaddedOutgoingBuffer.GetData(), OutgoingSignal, NumOutgoingSamples * sizeof(float));
-
-		int32 FilterLength = FMath::Max(NumIncomingSamples, NumOutgoingSamples);
-PRAGMA_DISABLE_DEPRECATION_WARNINGS
-		CrossCorrelate(ZeroPaddedIncomingBuffer.GetData(), ZeroPaddedOutgoingBuffer.GetData(), FilterLength, FFTSize, IncomingFrequencies, OutgoingFrequencies, OutWeights);
-PRAGMA_ENABLE_DEPRECATION_WARNINGS
-
-		// TODO: Perform linear phase conversion on OutWeights.
-
-	}
-
-	FAcousticEchoCancellation_DEPRECATED::FAcousticEchoCancellation_DEPRECATED(float InConvergenceRate, int32 CallbackSize, int32 InFilterLength, int32 InFilterUpdateRate /*= 1*/)
-		: AdaptiveFilter(InFilterLength, CallbackSize)
-		, ConvergenceRate(InConvergenceRate)
-		, FilterLength(InFilterLength)
-		, FilterUpdateRate(InFilterUpdateRate)
-		, FitlerUpdateCounter(InFilterUpdateRate)
-	{
-		checkSlow(FMath::IsPowerOfTwo(FilterLength));
-	}
-
-	void FAcousticEchoCancellation_DEPRECATED::ProcessAudio(float* InAudio, int32 NumSamples)
-	{
-		checkSlow(FMath::IsPowerOfTwo(NumSamples));
-
-		int32 NumBlocks = NumSamples / FilterLength;
-		for (int32 Index = 0; Index < NumBlocks; Index++)
-		{
-			float* BlockAudio = &InAudio[Index * FilterLength];
-
-			if (FitlerUpdateCounter == 0)
-			{
-				FilterComputerInput.Reset();
-				FilterComputerInput.AddZeroed(FilterLength);
-
-				// Since we are likely sampling audio buffers much less often than we are rendering them,
-				// we always skip ahead to the latest audio to compute our filter coefficients.
-				const bool bUseLatestAudio = true;
-
-				// Grab our audio from our taps:
-				int32 PopResult = PatchMixer.PopAudio(FilterComputerInput.GetData(), FilterLength, bUseLatestAudio);
-
-				if (PopResult != 0)
-				{
-					// Compute our new filter target and pass it to the adaptive filter:
-					FilterComputer.GenerateWeights(FilterComputerInput.GetData(), FilterLength, BlockAudio, FilterLength, FilterComputerOutput);
-					AdaptiveFilter.ProcessAudio(BlockAudio, NumSamples);
-					AdaptiveFilter.SetWeights(FilterComputerOutput, FilterLength, ConvergenceRate);
-				}
-				else
-				{
-					// If our taps were empty, don't both updating the filter.
-					AdaptiveFilter.ProcessAudio(BlockAudio, FilterLength);
-				}
-
-				FitlerUpdateCounter = FilterUpdateRate - 1;
-			}
-			else
-			{
-				AdaptiveFilter.ProcessAudio(BlockAudio, FilterLength);
-				FitlerUpdateCounter--;
-			}
-		}
-	}
-
-	Audio::FPatchInput FAcousticEchoCancellation_DEPRECATED::AddNewSignalPatch(int32 ExpectedLatency, float Gain /*= 1.0f*/)
-	{
-		return PatchMixer.AddNewInput(ExpectedLatency, Gain);
-	}
-
-	void FAcousticEchoCancellation_DEPRECATED::RemoveSignalPatch(const FPatchInput& Patch)
-	{
-		PatchMixer.RemovePatch(Patch);
-	}
 }
diff --git a/Engine/Source/Runtime/SignalProcessing/Public/DSP/AudioFFT.h b/Engine/Source/Runtime/SignalProcessing/Public/DSP/AudioFFT.h
--- a/Engine/Source/Runtime/SignalProcessing/Public/DSP/AudioFFT.h
+++ b/Engine/Source/Runtime/SignalProcessing/Public/DSP/AudioFFT.h
@@ -11,16 +11,13 @@
 namespace Audio
 {
 	// Types of spectrums which can be directly derived from FFTFreqDomainData
+	
 	enum class ESpectrumType : uint8
 	{
 		MagnitudeSpectrum,
 		PowerSpectrum
 	};
 
-	namespace FFTIntrinsics
-	{
-		SIGNALPROCESSING_API uint32 NextPowerOf2(uint32 Input);
-	}
 
 	enum class EWindowType : uint8
 	{
@@ -74,91 +71,8 @@ namespace Audio
 		SIGNALPROCESSING_API void Generate(int32 NumFrames, int32 NumChannels, bool bIsPeriodic);
 	};
 
-	struct FFTTimeDomainData_DEPRECATED
-	{
-		float* Buffer; // Pointer to a single channel of floats.
-		int32 NumSamples; // Number of samples in InBuffer divided by the number of channels. must be a power of 2.
-	};
-
-	/** AudioFFT Deprecation
-	 *
-	 * The AudioFFT is deprecated in favor of optimized software implementations 
-	 * and hardware implementations. The FFT implemenented in "PerformFFT(...)" 
-	 * has poor CPU performance as it computes the FFT radix weights each time
-	 * "PerformFFT(...)" is executed. 
-	 *
-	 * Better optimized and hardward FFT implementations can be accessed through the 
-	 * FFFTFactory class in FFTAlgorithm.h. That class will attempt to return the 
-	 * best FFT algorithm implementation available for the given FFFTSettings.
-	 */
-
-	struct UE_DEPRECATED(5.1, "Use IFFTAlgorithm and TArray for time domain data") FFTTimeDomainData;
-	struct FFTTimeDomainData : FFTTimeDomainData_DEPRECATED 
-	{
-		FFTTimeDomainData() = default;
-
-		FFTTimeDomainData(float* InBuffer, int32 InNumSamples)
-		{
-			Buffer = InBuffer;
-			NumSamples = InNumSamples;
-		}
-	};
-
-	struct FFTFreqDomainData_DEPRECATED
-	{
-		// arrays in which real and imaginary values will be populated.
-		float* OutReal; // Should point to an already allocated array of floats that is FFTInputParams::NumSamples long.
-		float* OutImag; // Should point to an already allocated array of floats that is FFTInputParams::NumSamples long.
-	};
-
-	struct UE_DEPRECATED(5.1, "Use IFFTAlgorithm and TArray for time freq domain data") FFTFreqDomainData;
-	struct FFTFreqDomainData : FFTFreqDomainData_DEPRECATED 
-	{
-		FFTFreqDomainData() = default;
-
-		FFTFreqDomainData(float* Real, float* Imag)
-		{
-			OutReal = Real;
-			OutImag = Imag;
-		}
-	};
-
-	// Performs a one-time FFT on a float buffer. Does not support complex signals.
-	// This function assumes that, if you desire a window for your FFT, that window was already
-	// applied to FFTInputParams.InBuffer.
-
- 	UE_DEPRECATED(5.1, "Use FVectorFFTRealToComplex or FFFTFactory instead.")
-	SIGNALPROCESSING_API void PerformFFT(const FFTTimeDomainData_DEPRECATED& InputParams, FFTFreqDomainData_DEPRECATED& OutputParams);
-
- 	UE_DEPRECATED(5.1, "Use FVectorFFTRealToComplex or FFFTFactory instead.")
-	SIGNALPROCESSING_API void PerformIFFT(FFTFreqDomainData_DEPRECATED& InputParams, FFTTimeDomainData_DEPRECATED& OutputParams);
-
-
-	// FFT Algorithm factory for this FFT implementation
-	class UE_DEPRECATED(5.1, "Use FVectorFFTFactory instead.") FAudioFFTAlgorithmFactory;
-	class FAudioFFTAlgorithmFactory : public IFFTAlgorithmFactory
-	{
-		public:
-			SIGNALPROCESSING_API virtual ~FAudioFFTAlgorithmFactory();
-
-			// Name of this fft algorithm factory. 
-			SIGNALPROCESSING_API virtual FName GetFactoryName() const override;
-
-			// If true, this implementation uses hardware acceleration.
-			SIGNALPROCESSING_API virtual bool IsHardwareAccelerated() const override;
-
-			// If true, this implementation requires input and output arrays to be 128 bit aligned.
-			SIGNALPROCESSING_API virtual bool Expects128BitAlignedArrays() const override;
-
-			// Returns true if the input settings are supported by this factory.
-			SIGNALPROCESSING_API virtual bool AreFFTSettingsSupported(const FFFTSettings& InSettings) const override;
-
-			// Create a new FFT algorithm.
-			SIGNALPROCESSING_API virtual TUniquePtr<IFFTAlgorithm> NewFFTAlgorithm(const FFFTSettings& InSettings) override;
-	};
-
-	struct FrequencyBuffer
- 	
+	
+	struct UE_DEPRECATED(5.5, "FrequencyBuffer is no longer used.") FrequencyBuffer
 	{
 		FAlignedFloatBuffer Real;
 		FAlignedFloatBuffer Imag;
@@ -187,95 +101,6 @@ namespace Audio
 		}
 	};
 
-	// Performs an acyclic FFT correlation on FirstBuffer and Second buffer and stores the output in OutCorrelation.
-	// If bCyclic is false, This function may zero pad FirstBuffer and Second Buffer as needed.
-	// If bCyclic is true, FirstBuffer and SecondBuffer should have the same length, and that length should be a power of two.
-	UE_DEPRECATED(5.1, "Cross correlate is no longer supported.")
-	SIGNALPROCESSING_API void CrossCorrelate(FAlignedFloatBuffer& FirstBuffer, FAlignedFloatBuffer& SecondBuffer, FAlignedFloatBuffer& OutCorrelation, bool bZeroPad = true);
-	UE_DEPRECATED(5.1, "Cross correlate is no longer supported.")
-	SIGNALPROCESSING_API void CrossCorrelate(FAlignedFloatBuffer& FirstBuffer, FAlignedFloatBuffer& SecondBuffer, FrequencyBuffer& OutCorrelation, bool bZeroPad = true);
-	UE_DEPRECATED(5.1, "Cross correlate is no longer supported.")
-	SIGNALPROCESSING_API void CrossCorrelate(const float* FirstBuffer, const float* SecondBuffer, int32 NumSamples, int32 FFTSize, float* OutCorrelation, int32 OutCorrelationSamples);
-	UE_DEPRECATED(5.1, "Cross correlate is no longer supported.")
-	SIGNALPROCESSING_API void CrossCorrelate(const float* FirstBuffer, const float* SecondBuffer, int32 NumSamples, int32 FFTSize, FrequencyBuffer& OutCorrelation);
-
-	// These variations do not allocate any additional memory during the function, provided that the FrequencyBuffers are already allocated.
-	UE_DEPRECATED(5.1, "Cross correlate is no longer supported.")
-	SIGNALPROCESSING_API void CrossCorrelate(const float* FirstBuffer, const float* SecondBuffer, int32 NumSamples, int32 FFTSize, FrequencyBuffer& FirstBufferFrequencies, FrequencyBuffer& SecondBufferFrequencies, FrequencyBuffer& OutCorrelation);
-	UE_DEPRECATED(5.1, "Cross correlate is no longer supported.")
-	SIGNALPROCESSING_API void CrossCorrelate(FrequencyBuffer& FirstBufferFrequencies, FrequencyBuffer& SecondBufferFrequencies, int32 NumSamples, FrequencyBuffer& OutCorrelation);
-
-	// Deprecated in 5.1
-	class FFFTConvolver_DEPRECATED
-	{
-	public:
-		SIGNALPROCESSING_API FFFTConvolver_DEPRECATED();
-
-		/*
-		 * Applies the convolver's internal window to InputAudio. Until SetWindow is called, ProcessAudio will not affect InputAudio.
-		 * InputAudio must be a power of two.
-		 */
-		SIGNALPROCESSING_API void ProcessAudio(float* InputAudio, int32 NumSamples);
-
-		/**
-		 * Resets the filter window. NOT thread safe to call during ProcessAudio.
-		 * This function can be called with a time domain impulse response, or precomputed frequency values. 
-		 * FilterSize must be a power of two.
-		 */ 
-		SIGNALPROCESSING_API void SetFilter(const float* InFilterReal, const float* InFilterImag, int32 FilterSize, int32 FFTSize);
-		SIGNALPROCESSING_API void SetFilter(const FrequencyBuffer& InFilterFrequencies, int32 FilterSize);
-		SIGNALPROCESSING_API void SetFilter(const float* TimeDomainBuffer, int32 FilterSize);
-		SIGNALPROCESSING_API void SetFilter(const FAlignedFloatBuffer& TimeDomainBuffer);
-
-	private:
-		void ConvolveBlock(float* InputAudio, int32 NumSamples);
-		void SumInCOLABuffer(float* InputAudio, int32 NumSamples);
-		void SetCOLABuffer(float* InAudio, int32 NumSamples);
-
-		FrequencyBuffer FilterFrequencies;
-		FrequencyBuffer InputFrequencies;
-		int32 BlockSize;
-
-
-		FAlignedFloatBuffer TimeDomainInputBuffer;
-		FAlignedFloatBuffer COLABuffer;
-	};
-
-	class UE_DEPRECATED(5.1, "Use FConvolutionFactory or UniformPartitionConvolutionFactory.") FFFTConvolver : public FFFTConvolver_DEPRECATED
-	{
-	};
-
-	// Computes the power spectrum from FFTFreqDomainData. Applies a 1/(FFTSize^2) scaling to the output to 
-	// maintain equal energy between original time domain data and output spectrum.  Only the first 
-	// (FFTSize / 2 + 1) spectrum values are calculated. These represent the frequencies from 0 to Nyquist.
-	//
-	// InFrequencyData is the input frequency domain data. Generally this is created by calling PerformFFT(...)
-	// FFTSize is the number of samples used when originally calculating the FFT
-	// OutBuffer is an aligned buffer which will contain spectrum data. It will constain (FFTSize / 2 + 1) elements.
-	UE_DEPRECATED(5.1, "Use equivalent methods in FloatArrayMath.h")
-	SIGNALPROCESSING_API void ComputePowerSpectrum(const FFTFreqDomainData_DEPRECATED& InFrequencyData, int32 FFTSize, FAlignedFloatBuffer& OutBuffer);
-
-	// Computes the magnitude spectrum from FFTFreqDomainData. Applies a 1/FFTSize scaling to the output to 
-	// maintain equal energy between original time domain data and output spectrum.  Only the first 
-	// (FFTSize / 2 + 1) spectrum values are calculated. These represent the frequencies from 0 to Nyquist.
-	//
-	// InFrequencyData is the input frequency domain data. Generally this is created by calling PerformFFT(...)
-	// FFTSize is the number of samples used when originally calculating the FFT
-	// OutBuffer is an aligned buffer which will contain spectrum data. It will constain (FFTSize / 2 + 1) elements.
-	UE_DEPRECATED(5.1, "Use equivalent methods in FloatArrayMath.h")
-	SIGNALPROCESSING_API void ComputeMagnitudeSpectrum(const FFTFreqDomainData_DEPRECATED& InFrequencyData, int32 FFTSize, FAlignedFloatBuffer& OutBuffer);
-	
-	// Computes the spectrum from FFTFreqDomainData. Applies a scaling to the output to maintain equal 
-	// energy between original time domain data and output spectrum.  Only the first (FFTSize / 2 + 1)
-	// spectrum values are calculated. These represent the frequencies from 0 to Nyquist.
-	//
-	// InSpectrumType denotes which spectrum type to calculate.
-	// InFrequencyData is the input frequency domain data. Generally this is created by calling PerformFFT(...)
-	// FFTSize is the number of samples used when originally calculating the FFT
-	// OutBuffer is an aligned buffer which will contain spectrum data. It will constain (FFTSize / 2 + 1) elements.
-	UE_DEPRECATED(5.1, "Use equivalent methods in FloatArrayMath.h")
-	SIGNALPROCESSING_API void ComputeSpectrum(ESpectrumType InSpectrumType, const FFTFreqDomainData_DEPRECATED& InFrequencyData, int32 FFTSize, FAlignedFloatBuffer& OutBuffer);
-
 	// Return the ceiling of the log2 of InNum
 	SIGNALPROCESSING_API int32 CeilLog2(int32 InNum);
 
diff --git a/Engine/Source/Runtime/SignalProcessing/Public/DSP/VoiceProcessing.h b/Engine/Source/Runtime/SignalProcessing/Public/DSP/VoiceProcessing.h
--- a/Engine/Source/Runtime/SignalProcessing/Public/DSP/VoiceProcessing.h
+++ b/Engine/Source/Runtime/SignalProcessing/Public/DSP/VoiceProcessing.h
@@ -127,123 +127,4 @@ namespace Audio
 		float GainMin;
 		float GainMax;
 	};
-	
-	/** VoiceProcessing.h Deprecation
-	 *
-	 * Several classes in VoiceProcessing are deprecated due to lack of support
-	 * and lack of need. Current voice processing solutions are available through 
-	 * EOS and WebRTC. 
-	 */
-	
-	/**
-	 * This filter takes a precomputed set of FIR weights in the frequency domain, and linearly converges to it.
-	 * If no new weights are given and we've converged to our previous input weights, this works like a normal FFT-based FIR filter.
-	 * Convergence is non-asymptotic: if no new weights are given after our number of steps until convergence, our filter is using the exact weights given.
-	 */
-	class FAdaptiveFilter_DEPRECATED // Deprecated in 5.1
-	{
-	public:
-		SIGNALPROCESSING_API FAdaptiveFilter_DEPRECATED(int32 FilterLength, int32 AudioCallbackSize);
-
-		/*
-		* Applies current filter to InAudio in-place. If there is a new target set of weights, they can be input below.
-		*/
-		SIGNALPROCESSING_API void ProcessAudio(float* InAudio, int32 NumSamples);
-
-		SIGNALPROCESSING_API void SetWeights(const FrequencyBuffer& InFilterWeights, int32 FilterLength, float InLearningRate);
-
-	private:
-
-		FAdaptiveFilter_DEPRECATED();
-
-		/** This is called every process callback to (A) reset our weight deltas if we have a new target, and (B) increment our current filter weights if we haven't converged yet. */
-		void AdaptFilter();
-
-		/*
-		 * Recomputes our weight deltas. Called in the ProcessAudio callback when ProcessAudio was called with a new set of weights.
-		 * Delta is computed 
-		*/
-		void SetWeightDeltas(const float* InWeightsReal, const float* InWeightsImag, int32 NumWeights, float InLearningRate);
-		void IncrementWeights();
-
-		FrequencyBuffer WeightDeltas;
-		FrequencyBuffer CurrentWeights;
-		int32 WindowSize;
-
-		FrequencyBuffer InputFrequencies;
-
-		int32 CurrentStepsUntilConvergence;
-
-		FFFTConvolver_DEPRECATED Convolver;
-	};
-
-	class UE_DEPRECATED(5.1, "FAdaptiveFilter will no longer be supported.") FAdaptiveFilter;
-	class FAdaptiveFilter : public FAdaptiveFilter_DEPRECATED
-	{
-		public:
-			using FAdaptiveFilter_DEPRECATED::FAdaptiveFilter_DEPRECATED;
-	};
-
-	/**
-	 * This class takes an incoming signal and an outgoing signal, Correlates them, and returns the frequency values of the weight targets to pass to an adaptive filter.  
-	 */
-	class FFDAPFilterComputer_DEPRECATED // Deprecated in 5.1
-	{
-	public:
-		SIGNALPROCESSING_API FFDAPFilterComputer_DEPRECATED();
-
-		SIGNALPROCESSING_API void GenerateWeights(const float* IncomingSignal, int32 NumIncomingSamples, const float* OutgoingSignal, int32 NumOutgoingSamples, FrequencyBuffer& OutWeights);
-
-	private:
-		FrequencyBuffer IncomingFrequencies;
-		FrequencyBuffer OutgoingFrequencies;
-
-		FAlignedFloatBuffer ZeroPaddedIncomingBuffer;
-		FAlignedFloatBuffer ZeroPaddedOutgoingBuffer;
-	};
-
-	class UE_DEPRECATED(5.1, "FAFDAPFFilterComputer will no longer be supported.") FFDAPFilterComputer;
-	class FFDAPFilterComputer : public FFDAPFilterComputer_DEPRECATED 
-	{
-	};
-
-	/*
-	 * This class uses an adaptive filter to cancel out any rendered audio signal that might be picked up by the mic.
-	 * To add a new patch to a rendered audio signal, user AddNewSignalPatch. See FPatchInput for how to push audio.
-	 * ProcessAudio then filters the microphone signal accordingly.
-	 */
-	class FAcousticEchoCancellation_DEPRECATED // Deprecated in 5.1
-	{
-	public:
-		/**
-		 * Convergence Rate should be a number between 0 and 1. The higher the number, the quicker the adaptive filter reacts. 
-		 */
-		SIGNALPROCESSING_API FAcousticEchoCancellation_DEPRECATED(float InConvergenceRate, int32 CallbackSize, int32 InFilterLength, int32 InFilterUpdateRate = 1);
-
-		/** Callback function for outgoing audio signal. This is where the filter is applied, and the bulk of the DSP work takes place. */
-		SIGNALPROCESSING_API void ProcessAudio(float* InAudio, int32 NumSamples);
-
-		/** This is how any signal that may be picked up by the microphone may be added to the echo cancellation here: */
-		SIGNALPROCESSING_API FPatchInput AddNewSignalPatch(int32 ExpectedLatency, float Gain = 1.0f);
-		SIGNALPROCESSING_API void RemoveSignalPatch(const FPatchInput& Patch);
-
-	private:
-		FPatchMixer PatchMixer;
-		FFDAPFilterComputer_DEPRECATED FilterComputer;
-		FAdaptiveFilter_DEPRECATED AdaptiveFilter;
-		
-		FAlignedFloatBuffer FilterComputerInput;
-		FrequencyBuffer FilterComputerOutput;
-		float ConvergenceRate;
-		int32 FilterLength;
-		int32 FilterUpdateRate;
-		int32 FitlerUpdateCounter;
-	};
-
-	class UE_DEPRECATED(5.1, "FAcousticEchoCancellation will no longer be supported.") FAcousticEchoCancellation;
-	class FAcousticEchoCancellation : public FAcousticEchoCancellation_DEPRECATED
-	{
-		public:
-			using FAcousticEchoCancellation_DEPRECATED::FAcousticEchoCancellation_DEPRECATED;
-	};
 }
