PR #12001: Fix Crash in SampleBuffer.h
#rnx

diff --git a/Engine/Source/Runtime/SignalProcessing/Public/SampleBuffer.h b/Engine/Source/Runtime/SignalProcessing/Public/SampleBuffer.h
--- a/Engine/Source/Runtime/SignalProcessing/Public/SampleBuffer.h
+++ b/Engine/Source/Runtime/SignalProcessing/Public/SampleBuffer.h
@@ -221,7 +221,7 @@ namespace Audio
 				else if constexpr(std::is_same_v<SampleType, float> && std::is_same_v<OtherSampleType, int16>)
 				{
 					// Convert from int to float:
-					Audio::ArrayPcm16ToFloat(MakeArrayView(InputBuffer, InNumSamples), MakeArrayView(&RawPCMData[StartIndex], NumSamples));
+					Audio::ArrayPcm16ToFloat(MakeArrayView(InputBuffer, InNumSamples), MakeArrayView(&RawPCMData[StartIndex], InNumSamples));
 				}
 				else
 				{
