[Metasounds] Added smoothing to delay nodes

diff --git a/Engine/Plugins/Runtime/Metasound/Source/MetasoundStandardNodes/Private/MetasoundDelayNode.cpp b/Engine/Plugins/Runtime/Metasound/Source/MetasoundStandardNodes/Private/MetasoundDelayNode.cpp
--- a/Engine/Plugins/Runtime/Metasound/Source/MetasoundStandardNodes/Private/MetasoundDelayNode.cpp
+++ b/Engine/Plugins/Runtime/Metasound/Source/MetasoundStandardNodes/Private/MetasoundDelayNode.cpp
@@ -56,6 +56,7 @@ namespace Metasound
 		virtual FDataReferenceCollection GetOutputs() const override;
 		void Reset(const IOperator::FResetParams& InParams);
 		void Execute();
+		void ExecuteInternal(int32 StartFrame, int32 EndFrame);
 
 	private:
 		float GetInputDelayTimeMsec() const;
@@ -174,13 +175,19 @@ namespace Metasound
 		TriggerReset->ExecuteBlock(
 			[&](int32 StartFrame, int32 EndFrame)
 			{
+				ExecuteInternal(StartFrame, EndFrame); 
 			},
-			[this](int32 StartFrame, int32 EndFrame)
+			 [&, this](int32 StartFrame, int32 EndFrame)
 			{
-				DelayBuffer.Reset();
+				FeedbackSample = 0.f;
+				DelayBuffer.ResetWithFade();
+			 	ExecuteInternal(StartFrame, EndFrame);
 			}
 		);
+	}
 
+	void FDelayOperator::ExecuteInternal(int32 StartFrame, int32 EndFrame)
+	{
 		// Get clamped delay time
 		float CurrentInputDelayTime = GetInputDelayTimeMsec();
 
@@ -205,7 +212,7 @@ namespace Metasound
 		{
 			FeedbackSample = 0.0f;
 
-			for (int32 FrameIndex = 0; FrameIndex < NumFrames; ++FrameIndex)
+			for (int32 FrameIndex = StartFrame; FrameIndex < EndFrame; ++FrameIndex)
 			{
 				OutputAudio[FrameIndex] = CurrentWetLevel * DelayBuffer.ProcessAudioSample(InputAudio[FrameIndex]) + CurrentDryLevel * InputAudio[FrameIndex];
 			}
@@ -213,7 +220,7 @@ namespace Metasound
 		else
 		{
 			// There is some amount of feedback so we do the feedback mixing
-			for (int32 FrameIndex = 0; FrameIndex < NumFrames; ++FrameIndex)
+			for (int32 FrameIndex = StartFrame; FrameIndex < EndFrame; ++FrameIndex)
 			{
 				OutputAudio[FrameIndex] = CurrentWetLevel * DelayBuffer.ProcessAudioSample(InputAudio[FrameIndex] + FeedbackSample * FeedbackAmount) + CurrentDryLevel * InputAudio[FrameIndex];
 				FeedbackSample = OutputAudio[FrameIndex];
diff --git a/Engine/Source/Runtime/SignalProcessing/Private/Delay.cpp b/Engine/Source/Runtime/SignalProcessing/Private/Delay.cpp
--- a/Engine/Source/Runtime/SignalProcessing/Private/Delay.cpp
+++ b/Engine/Source/Runtime/SignalProcessing/Private/Delay.cpp
@@ -2,6 +2,7 @@
 
 #include "DSP/Delay.h"
 #include "DSP/Dsp.h"
+#include "DSP/FloatArrayMath.h"
 #include "HAL/IConsoleManager.h"
 
 static float FDelayInitialAllocationSecondsCVar = -1.0f;
@@ -12,6 +13,13 @@ FAutoConsoleVariableRef CVarFDelayInitialAllocationSeconds(
 	//TEXT("The default is -1.  A value less than zero will allocate the full InBufferLengthSec\n"),
 	ECVF_Default);
 
+static int32 DelayResetFadeLengthSamplesCVar = 64;
+FAutoConsoleVariableRef CVarDelayResetFadeLengthSamples(
+	TEXT("au.DSP.DelayResetFadeLengthSamples"),
+	DelayResetFadeLengthSamplesCVar,
+	TEXT("Controls fade length (in samples) when clearing internal memory with ResetWithFade().\n"),
+	ECVF_Default);
+
 namespace Audio
 {
 	FDelay::FDelay()
@@ -59,6 +67,98 @@ namespace Audio
 
 		Update(true);
 	}
+
+	void FDelay::ResetWithFade()
+	{
+		if (DelayResetFadeLengthSamplesCVar == 0)
+		{
+			Reset();
+			return;
+		}			
+		
+		const int32 NumSamplesToFade = FMath::Min(DelayResetFadeLengthSamplesCVar,  static_cast<int32>(DelayInSamples - 1));
+		const bool bIsReadRegionSplit = WriteIndex < ReadIndex;
+		
+		// we don't want to accept new input while the delay is faded
+		// (this could cause the two fades to create a small chirp)
+		// so we start the input gain negative and rely on the clamping to mute (gain = 0)
+		InputFadeGainStep = 1.f / static_cast<float>(NumSamplesToFade);
+		InputAttenuation = -1.2f;
+
+
+		// simple case
+		if (!bIsReadRegionSplit)
+		{
+			const int32 NumSamplesToZero = WriteIndex - ReadIndex - NumSamplesToFade;
+			
+			const TArrayView<float> FadeRegion(AudioBuffer.GetData() + ReadIndex, NumSamplesToFade);
+			const TArrayView<float> ZeroRegion(AudioBuffer.GetData() + ReadIndex + NumSamplesToFade, NumSamplesToZero);
+			
+			ArrayFade(FadeRegion, 1.0f, 0.0f);
+			FMemory::Memzero(ZeroRegion.GetData(), ZeroRegion.NumBytes());
+			
+			return;
+		}
+
+		// complex case, two regions to worry about
+		const int32 NumSamplesToZero = (AudioBufferSize - ReadIndex) + WriteIndex - NumSamplesToFade;
+		const int32 ReadRegion1Num = AudioBufferSize - ReadIndex;
+		const int32 ReadRegion2Num = WriteIndex;
+
+		// 3 complex cases:
+		// 1.) single fade region and two zero regions (fade region < read region 1)
+		// 2.) two fade regions and a single zero region (fade region > read region 2)
+		// 3.) [else] read region 1 is exactly the fade region (unlikely)
+		// note: we should have at least on sample to zero out
+		//       (see -1 in calculation of NumSamplesToFade)
+		//       so technically we don't have an "all fade" case
+
+		// case 1:
+		if (NumSamplesToFade < ReadRegion1Num)
+		{
+			const int32 ZeroRegion1Size = NumSamplesToZero - ReadRegion2Num;
+			const int32 ZeroRegion2Size = NumSamplesToZero - ZeroRegion1Size;
+
+			const TArrayView<float> FadeRegion(AudioBuffer.GetData() + ReadIndex, NumSamplesToFade);
+			const TArrayView<float> ZeroRegion1(AudioBuffer.GetData() + ReadIndex + NumSamplesToFade, ZeroRegion1Size);
+			const TArrayView<float> ZeroRegion2(AudioBuffer.GetData(), ZeroRegion2Size);
+
+			ArrayFade(FadeRegion, 1.0f, 0.0f);
+			FMemory::Memzero(ZeroRegion1.GetData(), ZeroRegion1.NumBytes());
+			FMemory::Memzero(ZeroRegion2.GetData(), ZeroRegion2.NumBytes());
+
+			return;
+		}
+		
+		// case 2:
+		if (NumSamplesToFade > ReadRegion1Num)
+		{
+			const int32 FadeRegion1Size = AudioBufferSize - ReadIndex;
+			const int32 FadeRegion2Size = NumSamplesToFade - FadeRegion1Size;
+			const int32 ZeroRegionSize = WriteIndex - FadeRegion2Size;
+
+			const TArrayView<float> FadeRegion1(AudioBuffer.GetData() + ReadIndex, FadeRegion1Size);
+			const TArrayView<float> FadeRegion2(AudioBuffer.GetData(), FadeRegion2Size);
+			const TArrayView<float> ZeroRegion(AudioBuffer.GetData() + FadeRegion2Size, ZeroRegionSize);
+
+			// technically there will be 2 samples with this gain value applied
+			// (should be imperceptible)
+			const float MidGain = FadeRegion2Size / static_cast<float>(NumSamplesToFade);
+
+			ArrayFade(FadeRegion1, 1.0f, MidGain);
+			ArrayFade(FadeRegion2, MidGain, 0.0f);
+			FMemory::Memzero(ZeroRegion.GetData(), ZeroRegion.NumBytes());
+
+			return;
+		}
+
+		// case 3 (rare):
+		const TArrayView<float> FadeRegion(AudioBuffer.GetData() + ReadIndex, AudioBufferSize - ReadIndex);
+		const TArrayView<float> ZeroRegion(AudioBuffer.GetData(), WriteIndex);
+
+		ArrayFade(FadeRegion, 1.0f, 0.0f);
+		FMemory::Memzero(ZeroRegion.GetData(), ZeroRegion.NumBytes());
+	}
 	 
 	void FDelay::SetDelayMsec(const float InDelayMsec)
 	{
@@ -158,8 +258,12 @@ namespace Audio
 
 	void FDelay::WriteDelayAndInc(const float InDelayInput)
 	{
+		// update input gain
+		InputAttenuation = FMath::Min(InputAttenuation + InputFadeGainStep, 1.f); // Ceil to 1.0
+		float InputAttenuationTemp = FMath::Max(InputAttenuation, 0.f); // Floor to 0.0
+		
 		// write to the delay line
-		AudioBuffer[WriteIndex] = InDelayInput; // external feedback sample
+		AudioBuffer[WriteIndex] = InDelayInput * InputAttenuationTemp; // external feedback sample
 												// increment the pointers and wrap if necessary
 		WriteIndex++;
 		if (WriteIndex >= AudioBufferSize)
diff --git a/Engine/Source/Runtime/SignalProcessing/Public/DSP/Delay.h b/Engine/Source/Runtime/SignalProcessing/Public/DSP/Delay.h
--- a/Engine/Source/Runtime/SignalProcessing/Public/DSP/Delay.h
+++ b/Engine/Source/Runtime/SignalProcessing/Public/DSP/Delay.h
@@ -25,6 +25,8 @@ namespace Audio
 		// Resets the delay line state, flushes buffer and resets read/write pointers.
 		// (called by Init())
 		SIGNALPROCESSING_API void Reset();
+		SIGNALPROCESSING_API void ResetWithFade();
+		
 
 		// Sets the delay line length. Will clamp to within range of the max initialized delay line length (won't resize).
 		SIGNALPROCESSING_API void SetDelayMsec(const float InDelayMsec);
@@ -92,6 +94,10 @@ namespace Audio
 		// Output attenuation value.
 		float OutputAttenuation;
 
+		// Used to do a quick fade-in of input after a call to "ResetWithFade()"
+		float InputAttenuation = 1.f;
+		float InputFadeGainStep = 0.f;
+
 		// Attenuation in decibel
 		float OutputAttenuationDB;
 	};
