Fix to FDelay off-by-one (bad access) exposed dy ResetWithFade()
#rnx

diff --git a/Engine/Source/Runtime/SignalProcessing/Private/Delay.cpp b/Engine/Source/Runtime/SignalProcessing/Private/Delay.cpp
--- a/Engine/Source/Runtime/SignalProcessing/Private/Delay.cpp
+++ b/Engine/Source/Runtime/SignalProcessing/Private/Delay.cpp
@@ -28,6 +28,7 @@ namespace Audio
 		, WriteIndex(0)
 		, SampleRate(0)
 		, DelayInSamples(0.0f)
+		, MaxBufferLengthSamples(0)
 		, EaseDelayMsec(0.0f, 0.0001f)
 		, OutputAttenuation(1.0f)
 		, OutputAttenuationDB(0.0f)
@@ -70,13 +71,14 @@ namespace Audio
 
 	void FDelay::ResetWithFade()
 	{
-		if (DelayResetFadeLengthSamplesCVar == 0)
+		const int32 NumSamplesToFade = FMath::Min(DelayResetFadeLengthSamplesCVar,  static_cast<int32>(DelayInSamples - 1));
+		
+		if (DelayResetFadeLengthSamplesCVar == 0 || NumSamplesToFade < 2)
 		{
 			Reset();
 			return;
-		}			
+		}
 		
-		const int32 NumSamplesToFade = FMath::Min(DelayResetFadeLengthSamplesCVar,  static_cast<int32>(DelayInSamples - 1));
 		const bool bIsReadRegionSplit = WriteIndex < ReadIndex;
 		
 		// we don't want to accept new input while the delay is faded
@@ -316,7 +318,7 @@ namespace Audio
 			DelayInSamples = FMath::Clamp(DelayInSamples, 0.0f, (float)(AudioBufferSize - 1));
 
 			// Subtract from write index the delay in samples (will do interpolation during read)
-			ReadIndex = WriteIndex - (int32)(DelayInSamples + 1.0f);
+			ReadIndex = WriteIndex - (int32)(DelayInSamples);
 
 			// If negative, wrap around
 			if (ReadIndex < 0)
