Fix for crash in MetaSound Delay node reset
#rnx

diff --git a/Engine/Source/Runtime/SignalProcessing/Private/Delay.cpp b/Engine/Source/Runtime/SignalProcessing/Private/Delay.cpp
--- a/Engine/Source/Runtime/SignalProcessing/Private/Delay.cpp
+++ b/Engine/Source/Runtime/SignalProcessing/Private/Delay.cpp
@@ -87,7 +87,7 @@ namespace Audio
 		// (this could cause the two fades to create a small chirp)
 		// so we start the input gain negative and rely on the clamping to mute (gain = 0)
 		InputFadeGainStep = 1.f / static_cast<float>(NumSamplesToFade);
-		InputAttenuation = -1.2f;
+		InputAttenuation = -1.0f - InputFadeGainStep;
 
 
 		// simple case
@@ -168,14 +168,14 @@ namespace Audio
 	{
 		// Directly set the delay
 		const float NewDelayInSamples = InDelayMsec * SampleRate * 0.001f;
-		DelayInSamples = FMath::Min(NewDelayInSamples, MaxBufferLengthSamples);
+		DelayInSamples = FMath::Min(NewDelayInSamples, MaxBufferLengthSamples - 1);
 		ResizeIfNeeded(DelayInSamples);
 		Update(true);
 	}
 
 	void FDelay::SetDelaySamples(const float InDelaySamples)
 	{
-		DelayInSamples = FMath::Min(InDelaySamples, MaxBufferLengthSamples);
+		DelayInSamples = FMath::Min(InDelaySamples, MaxBufferLengthSamples - 1);
 		ResizeIfNeeded(DelayInSamples);
 		Update(true);
 	}
@@ -183,7 +183,7 @@ namespace Audio
 	void FDelay::SetEasedDelayMsec(const float InDelayMsec, const bool bIsInit)
 	{
 		const float DesiredDelayInSamples = InDelayMsec * SampleRate * 0.001f;
-		const float TargetDelayInSamples = FMath::Min(DesiredDelayInSamples, MaxBufferLengthSamples);
+		const float TargetDelayInSamples = FMath::Min(DesiredDelayInSamples, MaxBufferLengthSamples - 1);
 		ResizeIfNeeded(TargetDelayInSamples);
 
 		EaseDelayMsec.SetValue(InDelayMsec, bIsInit);
@@ -268,7 +268,8 @@ namespace Audio
 		
 		// write to the delay line
 		AudioBuffer[WriteIndex] = InDelayInput * InputAttenuationTemp; // external feedback sample
-												// increment the pointers and wrap if necessary
+		
+		// increment the pointers and wrap if necessary
 		WriteIndex++;
 		if (WriteIndex >= AudioBufferSize)
 		{
@@ -317,10 +318,10 @@ namespace Audio
 				DelayInSamples = EaseDelayMsec.GetNextValue() * SampleRate * 0.001f;
 			}
 
-			DelayInSamples = FMath::Clamp(DelayInSamples, 0.0f, (float)(AudioBufferSize - 1));
+			DelayInSamples = FMath::Clamp(DelayInSamples, 0.0f, (float)(MaxBufferLengthSamples - 1));
 
 			// Subtract from write index the delay in samples (will do interpolation during read)
-			ReadIndex = WriteIndex - (int32)(DelayInSamples);
+			ReadIndex = WriteIndex - (int32)DelayInSamples;
 
 			// If negative, wrap around
 			if (ReadIndex < 0)
diff --git a/Engine/Source/Runtime/SignalProcessing/Public/DSP/Delay.h b/Engine/Source/Runtime/SignalProcessing/Public/DSP/Delay.h
--- a/Engine/Source/Runtime/SignalProcessing/Public/DSP/Delay.h
+++ b/Engine/Source/Runtime/SignalProcessing/Public/DSP/Delay.h
@@ -94,12 +94,15 @@ namespace Audio
 		// Output attenuation value.
 		float OutputAttenuation;
 
-		// Used to do a quick fade-in of input after a call to "ResetWithFade()"
-		float InputAttenuation = 1.f;
-		float InputFadeGainStep = 0.f;
 
 		// Attenuation in decibel
 		float OutputAttenuationDB;
+
+	private:
+
+		// Used to do a quick fade-in of input after a call to "ResetWithFade()"
+		float InputAttenuation = 1.f;
+		float InputFadeGainStep = 0.f;
 	};
 
 }
