Add a ClearAudio method to FPatchInput to clear any queued audio.

diff --git a/Engine/Source/Runtime/SignalProcessing/Private/MultithreadedPatching.cpp b/Engine/Source/Runtime/SignalProcessing/Private/MultithreadedPatching.cpp
--- a/Engine/Source/Runtime/SignalProcessing/Private/MultithreadedPatching.cpp
+++ b/Engine/Source/Runtime/SignalProcessing/Private/MultithreadedPatching.cpp
@@ -106,6 +106,11 @@ namespace Audio
 		return NumSamplesPushed;
 	}
 
+	void FPatchOutput::EmptyInternalBuffer()
+	{
+		// Don't do an Empty() because then we would need to re-initialize.
+		InternalBuffer.Pop(GetNumSamplesAvailable());
+	}
 
 	int32 FPatchOutput::MixInAudio(float* OutBuffer, int32 NumSamples, bool bUseLatestAudio)
 	{
@@ -298,6 +303,16 @@ namespace Audio
 		return SamplesPushed;
 	}
 
+	void FPatchInput::ClearAudio()
+	{
+		if (!OutputHandle.IsValid())
+		{
+			return;
+		}
+
+		OutputHandle->EmptyInternalBuffer();
+	}
+
 	/** Returns the current number of samples buffered in this input. */
 	int32 FPatchInput::GetNumSamplesAvailable() const
 	{
diff --git a/Engine/Source/Runtime/SignalProcessing/Public/DSP/MultithreadedPatching.h b/Engine/Source/Runtime/SignalProcessing/Public/DSP/MultithreadedPatching.h
--- a/Engine/Source/Runtime/SignalProcessing/Public/DSP/MultithreadedPatching.h
+++ b/Engine/Source/Runtime/SignalProcessing/Public/DSP/MultithreadedPatching.h
@@ -61,6 +61,8 @@ namespace Audio
 		
 		SIGNALPROCESSING_API int32 PushAudioToInternalBuffer(const float* InBuffer, int32 NumSamples);
 		
+		SIGNALPROCESSING_API void EmptyInternalBuffer();
+
 		// Internal buffer.
 		TCircularAudioBuffer<float> InternalBuffer;
 
@@ -111,6 +113,9 @@ namespace Audio
 		 */
 		SIGNALPROCESSING_API int32 PushAudio(const float* InBuffer, int32 NumSamples);
 
+		/** Clears any audio queued up in the patch. */
+		SIGNALPROCESSING_API void ClearAudio();
+
 		/** Returns the current number of samples buffered in this input. */
 		SIGNALPROCESSING_API int32 GetNumSamplesAvailable() const;
 
