Fix for missing initializer
#rnx

diff --git a/Engine/Source/Runtime/SignalProcessing/Private/LFO.cpp b/Engine/Source/Runtime/SignalProcessing/Private/LFO.cpp
--- a/Engine/Source/Runtime/SignalProcessing/Private/LFO.cpp
+++ b/Engine/Source/Runtime/SignalProcessing/Private/LFO.cpp
@@ -17,6 +17,7 @@ namespace Audio
 		, LastOutput(0.0f)
 		, LoopCount(0.0f)
 		, QuadLastOutput(0.0f)
+		, PhaseOffset(0.0f)
 		, bBipolar(true)
 	{
 	}
