[Experimental] Refactor to Move CAT/Allocators/ScratchBuffers out of MetasoundExpermental to new module
Add TypeFamily and ChannelTypeFamily
Make+Break nodes
Node config support

diff --git a/Engine/Source/Runtime/SignalProcessing/Private/MultiMono.cpp b/Engine/Source/Runtime/SignalProcessing/Private/MultiMono.cpp
new file mode 100644
--- /dev/null
+++ b/Engine/Source/Runtime/SignalProcessing/Private/MultiMono.cpp
@@ -0,0 +1,58 @@
+// Copyright Epic Games, Inc. All Rights Reserved.
+
+#include "DSP/MultiMono.h"
+#include "DSP/FloatArrayMath.h"
+
+namespace Audio
+{
+	void MultiMonoMixUpOrDown(
+		TArrayView<const float> InSrc, const TArrayView<float> InDst, const int32 NumFrames, TArrayView<const float> MixGains,
+			const int32 NumSrcChannels, const int32 NumDstChannels)
+	{
+		MultiMonoMixUpOrDown(
+			MakeMultiMonoPointersFromView(InSrc, NumFrames, NumSrcChannels),
+			MakeMultiMonoPointersFromView(InDst, NumFrames, NumDstChannels),
+			NumFrames, 
+			MixGains);
+	}
+
+	static void MixMono(TArrayView<const float> InSrc, TArrayView<float> InDst, const float InGain)
+	{
+		const float* Src = InSrc.GetData();
+		float* Dst = InDst.GetData();
+		const int32 Num = InSrc.Num();
+		for (int32 i = 0; i < Num; i++,Src++,Dst++)
+		{
+			*Dst += *Src * InGain;
+		}
+	}
+
+	void MultiMonoMixUpOrDown(TArrayView<const float*> InSrc, TArrayView<float*> InDst, const int32 NumFrames,TArrayView<const float> MixGains)
+	{
+		const int32 NumSrcChannels = InSrc.Num();
+		const int32 NumDstChannels = InDst.Num();
+		checkSlow(NumFrames > 0);
+		checkSlow(MixGains.Num() == NumSrcChannels * NumDstChannels);
+		
+		for (int32 DstCh = 0; DstCh < NumDstChannels; ++DstCh)
+		{
+			for (int32 SrcCh = 0; SrcCh < NumSrcChannels; ++SrcCh)
+			{
+				if (const float ChannelGain = MixGains[DstCh + (SrcCh * NumDstChannels)]; !FMath::IsNearlyZero(ChannelGain))
+				{
+					MixMono(
+					MakeArrayView(InSrc[SrcCh], NumFrames),
+					MakeArrayView(InDst[DstCh], NumFrames),
+						ChannelGain);
+					
+					/*
+					ArrayMixIn(
+						MakeArrayView(InSrc[SrcCh], NumFrames),
+						MakeArrayView(InDst[DstCh], NumFrames),
+						ChannelGain);
+						*/
+				}
+			}
+		}	
+	}
+}
diff --git a/Engine/Source/Runtime/SignalProcessing/Public/DSP/MultiMono.h b/Engine/Source/Runtime/SignalProcessing/Public/DSP/MultiMono.h
new file mode 100644
--- /dev/null
+++ b/Engine/Source/Runtime/SignalProcessing/Public/DSP/MultiMono.h
@@ -0,0 +1,56 @@
+// Copyright Epic Games, Inc. All Rights Reserved.
+
+#pragma once
+
+#include "HAL/Platform.h"
+#include "Containers/ArrayView.h"
+#include "Containers/StaticArray.h"
+
+// Collection of multi-mono DSP algorithms. Although technically most of these are just setup to call the mono versions, there is enough
+// boilerplate to warrant a single lib.
+
+// Example Multi-mono-layout: [L][L][L][L][L][R][R][R][R][R], i.e. an entire channels worth of frames consecutively in memory, next
+// to the next channel and so on.
+
+namespace Audio
+{
+	constexpr int32 MaxNumTranscodeChannels = 128; // Some relatively high number. These are just pointers on stack.
+
+	// Stack array of pointers alias.
+	template<typename T> using TStackArrayOfPointers = TStaticArray<T*,MaxNumTranscodeChannels>;
+
+	// Helper to create array of pointers from multi-mono.
+	// Warning Array of pointers will be larger than you need, but filled to InNumChannels
+	template<typename T> TStackArrayOfPointers<T> MakeMultiMonoPointersFromView(TArrayView<T> InMultiMono, const int32 InNumFrames, const int32 InNumChannels)
+	{
+		checkSlow(InNumChannels <= MaxNumTranscodeChannels);
+		checkSlow(InNumChannels > 0);
+		checkSlow(InMultiMono.GetData());
+		checkSlow(InMultiMono.Num() % InNumFrames == 0 && InMultiMono.Num() >= InNumFrames);
+		TStackArrayOfPointers<T> ArrayOfPointers = {};
+		for (int32 i = 0; i < InNumChannels; ++i)
+		{
+			ArrayOfPointers[i] = InMultiMono.GetData() + (InNumFrames * i);
+		}
+		return ArrayOfPointers; // RTO so this should just inline
+	}
+	
+	/**
+	 * Given a Matrix of Gains (produced by calling ChannelMap.cpp (Audio::Create2DChannelMap), Mix Up/Down Source into Destination.
+	 * @param InSrc ArrayView to Source (Multi-mono)
+	 * @param InDst ArrayView to Destination (Mult-mono).
+	 * @param NumFrames Number of Frames in Each Channel.
+	 * @param MixGains Matrix of Gains (in Row Major Format, see Create2DChannelMap). [NumDstChannels*NumSrcChannels]
+	 * @param NumDstChannels Num Of Destination Channels.
+	 * @param NumSrcChannels Num Of Source Channels.
+	*/
+	SIGNALPROCESSING_API void MultiMonoMixUpOrDown(
+		TArrayView<const float> InSrc, const TArrayView<float> InDst, const int32 NumFrames, TArrayView<const float> MixGains,
+			const int32 NumSrcChannels, const int32 NumDstChannels);
+
+	SIGNALPROCESSING_API void MultiMonoMixUpOrDown(
+		TArrayView<const float*> InSrc, TArrayView<float*> InDst, const int32 NumFrames,TArrayView<const float> MixGains);
+	
+}
+
+
