New universal resampler implementation

diff --git a/Engine/Source/Runtime/AudioMixer/Private/AudioMixerSourceManager.cpp b/Engine/Source/Runtime/AudioMixer/Private/AudioMixerSourceManager.cpp
--- a/Engine/Source/Runtime/AudioMixer/Private/AudioMixerSourceManager.cpp
+++ b/Engine/Source/Runtime/AudioMixer/Private/AudioMixerSourceManager.cpp
@@ -13,7 +13,6 @@
 #include "DSP/FloatArrayMath.h"
 #include "IAudioExtensionPlugin.h"
 #include "AudioMixer.h"
-#include "samplerate.h"
 #include "Sound/SoundModulationDestination.h"
 #include "SoundFieldRendering.h"
 #include "ProfilingDebugging/CsvProfiler.h"
@@ -382,17 +381,6 @@ namespace Audio
 		}
 
 		FPlatformProcess::ReturnSynchEventToPool(CommandsProcessedEvent);
-
-		for (FSourceInfo& SourceInfo : SourceInfos)
-		{
-			ReleaseResampler(SourceInfo);
-		}
-
-		for (int32 i = 0; i < ResamplerPool.Num(); ++i)
-		{
-			src_delete(ResamplerPool[i]);
-		}
-		ResamplerPool.Reset();
 	}
 
 	void FMixerSourceManager::Init(const FSourceManagerInitParams& InitParams)
@@ -488,6 +476,7 @@ namespace Audio
 			SourceInfo.bIsVorbis = false;
 			SourceInfo.bHasPreDistanceAttenuationSend = false;
 			SourceInfo.bModFiltersUpdated = false;
+			SourceInfo.bResampling = false;
 
 #if AUDIO_MIXER_ENABLE_DEBUG_MODE
 			SourceInfo.bIsDebugMode = false;
@@ -496,8 +485,6 @@ namespace Audio
 			SourceInfo.NumInputChannels = 0;
 			SourceInfo.NumPostEffectChannels = 0;
 			SourceInfo.NumInputFrames = 0;
-
-			SourceInfo.Resampler = nullptr;
 		}
 		
 		GameThreadInfo.bIsBusy.AddDefaulted(NumTotalSources);
@@ -1043,6 +1030,7 @@ namespace Audio
 			SourceInfo.ActiveSoundPlayOrder = InitParams.ActiveSoundPlayOrder;
 			SourceInfo.AudioComponentID = InitParams.AudioComponentID;
 			SourceInfo.bIsSoundfield = InitParams.bIsSoundfield;
+			SourceInfo.bResampling = false;
 
 			// Call initialization from the render thread so anything wanting to do any initialization here can do so (e.g. procedural sound waves)
 			SourceInfo.SourceListener = InitParams.SourceListener;
@@ -1051,7 +1039,7 @@ namespace Audio
 			SourceInfo.NumInputChannels = InitParams.NumInputChannels;
 			SourceInfo.NumInputFrames = InitParams.NumInputFrames;
 
-			ReleaseResampler(SourceInfo);
+			SourceInfo.Resampler.Reset(SourceInfo.NumInputChannels);
 
 			// init and zero-out buffers
 			const int32 BufferSize = NumOutputFrames * InitParams.NumInputChannels;
@@ -2440,6 +2428,7 @@ namespace Audio
 			}
 			else // PerSourceResampling
 			{
+				const bool bIsFirstRun = !SourceInfo.bHasStarted;
 				SourceInfo.bHasStarted = true;
 
 				// Fill PreDistanceAttenuationBuffer with incoming source data
@@ -2494,9 +2483,14 @@ namespace Audio
 							}
 						}
 
-						if (FinalPitch != 1 && !SourceInfo.Resampler)
+						if (FinalPitch != 1 && !SourceInfo.bResampling)
 						{
-							AcquireResampler(SourceInfo);
+							SourceInfo.bResampling = true;
+							if (bIsFirstRun)
+							{
+								// Immediately start at the requested playback rate
+								SourceInfo.Resampler.SetFrameRatio(FinalPitch, 0);
+							}
 						}
 
 						// Pump as many samples as possibe from CurrentPCMBuffer to PreDistanceAttenuationBuffer
@@ -2505,23 +2499,22 @@ namespace Audio
 						float* CopyDestination = SourceInfo.PreDistanceAttenuationBuffer.GetData() + DestinationOffset;
 						const float* CopySource = SourceInfo.CurrentPCMBuffer->GetData() + SourceInfo.CurrentFrameIndex * SourceInfo.NumInputChannels;
 
-						if (SourceInfo.Resampler)
+						if (SourceInfo.bResampling)
 						{
 							// Resample if required
-							SRC_DATA SampleRateParams{};
-							SampleRateParams.data_in = CopySource;
-							SampleRateParams.data_out = CopyDestination;
-							SampleRateParams.input_frames = CurrentChunkFrames - SourceInfo.CurrentFrameIndex;
-							SampleRateParams.output_frames = FramesToWrite;
-							check(FinalPitch > 0);
-							SampleRateParams.src_ratio = 1 / FinalPitch;
-
-							src_process(SourceInfo.Resampler, &SampleRateParams);
+							int32 NumFramesConsumed = 0;
+							int32 NumFramesProduced = 0;
+							SourceInfo.Resampler.SetFrameRatio(FinalPitch, NumOutputFrames);
+							SourceInfo.Resampler.ProcessInterleaved(
+								TConstArrayView<float>{ CopySource, SourceInfo.NumInputChannels* (CurrentChunkFrames - SourceInfo.CurrentFrameIndex) },
+								TArrayView<float>{ CopyDestination, SourceInfo.NumInputChannels* FramesToWrite },
+								NumFramesConsumed, NumFramesProduced);
 
 							// Update bookkeeping
-							SourceInfo.CurrentFrameIndex += SampleRateParams.input_frames_used;
-							SourceInfo.NumFramesPlayed += SampleRateParams.input_frames_used;
-							FramesToWrite -= SampleRateParams.output_frames_gen;
+							check(NumFramesConsumed > 0 || NumFramesProduced > 0); // Make sure we don't get completely stuck
+							SourceInfo.CurrentFrameIndex += NumFramesConsumed;
+							SourceInfo.NumFramesPlayed += NumFramesConsumed;
+							FramesToWrite -= NumFramesProduced;
 						}
 						else
 						{
@@ -3950,49 +3943,4 @@ namespace Audio
 			UE_LOG(LogAudioMixer, Warning, TEXT("***** ThreadStackWalkAndDump for ThreadId(%lu) ******\n%s"), InThreadId, ANSI_TO_TCHAR(StackTrace));
 		}
 	}
-
-	void FMixerSourceManager::AcquireResampler(FSourceInfo& SourceInfo)
-	{
-		check(!SourceInfo.Resampler);
-
-		{
-			UE::TScopeLock<UE::FSpinLock> ScopeLock(ResamplerPoolLock);
-
-			for (int32 i = 0; i < ResamplerPool.Num(); ++i)
-			{
-				SRC_STATE* Resampler = ResamplerPool[i];
-				check(Resampler);
-				if (src_get_channels(Resampler) == SourceInfo.NumInputChannels)
-				{
-					// Reuse this resampler
-					ResamplerPool.RemoveAtSwap(i);
-					SourceInfo.Resampler = Resampler;
-					break;
-				}
-			}
-		}
-
-		if (SourceInfo.Resampler)
-		{
-			// Reset and return the reusable resampler
-			src_reset(SourceInfo.Resampler);
-			return;
-		}
-
-		// Allocate a new resampler
-		int Error = 0;
-		SourceInfo.Resampler = src_new(SRC_LINEAR, SourceInfo.NumInputChannels, &Error);
-		checkf(SourceInfo.Resampler, TEXT("src_new failed with error code %i"), Error);
-	}
-
-	void FMixerSourceManager::ReleaseResampler(FSourceInfo& SourceInfo)
-	{
-		SRC_STATE* Resampler = SourceInfo.Resampler;
-		if (Resampler)
-		{
-			SourceInfo.Resampler = nullptr;
-			UE::TScopeLock<UE::FSpinLock> ScopeLock(ResamplerPoolLock);
-			ResamplerPool.Add(Resampler);
-		}
-	}
 }
\ No newline at end of file
diff --git a/Engine/Source/Runtime/AudioMixer/Private/AudioMixerSourceManager.h b/Engine/Source/Runtime/AudioMixer/Private/AudioMixerSourceManager.h
--- a/Engine/Source/Runtime/AudioMixer/Private/AudioMixerSourceManager.h
+++ b/Engine/Source/Runtime/AudioMixer/Private/AudioMixerSourceManager.h
@@ -13,6 +13,7 @@
 #include "DSP/EnvelopeFollower.h"
 #include "DSP/InterpolatedOnePole.h"
 #include "DSP/ParamInterpolator.h"
+#include "DSP/RuntimeResampler.h"
 #include "IAudioExtensionPlugin.h"
 #include "ISoundfieldFormat.h"
 #include "Misc/SpinLock.h"
@@ -63,9 +64,6 @@ enum class ESourceManagerRenderThreadPhase: uint8
 	Finished,
 };
 
-// Forward declaration from libSampleRate
-typedef struct SRC_STATE_tag SRC_STATE;
-
 namespace Audio
 {
 	class FMixerSubmix;
@@ -440,7 +438,7 @@ namespace Audio
 		
 		struct FSourceInfo
 		{
-			FSourceInfo() {}
+			FSourceInfo() : Resampler(1 /* NumChannels */) {}
 			~FSourceInfo() {}
 
 			// Object which handles source buffer decoding
@@ -559,6 +557,7 @@ namespace Audio
 			uint8 bHasPreDistanceAttenuationSend:1;
 			uint8 bModFiltersUpdated : 1;
 			uint8 bShouldSourceBufferListenerZeroBuffer : 1;
+			uint8 bResampling : 1;
 
 			// Source format info
 			int32 NumInputChannels;
@@ -572,7 +571,7 @@ namespace Audio
 			uint64 AudioComponentID;
 
 			// Resampler instance
-			SRC_STATE* Resampler;
+			FRuntimeResampler Resampler;
 
 			FORCEINLINE void ResetModulators(const Audio::FDeviceId InDeviceId)
 			{
@@ -630,9 +629,6 @@ namespace Audio
 		void LogCallstacks();
 		void LogCallstack(uint32 InThreadId);
 
-		void AcquireResampler(FSourceInfo& SourceInfo);
-		void ReleaseResampler(FSourceInfo& SourceInfo);
-
 		// Array of listener transforms
 		TArray<FTransform> ListenerTransforms;
 
@@ -698,9 +694,6 @@ namespace Audio
 
 		TMpscQueue<FPendingAudioBusConnection> PendingAudioBusConnections;
 
-		UE::FSpinLock ResamplerPoolLock;
-		TArray<SRC_STATE*> ResamplerPool;
-
 		friend class FMixerSourceVoice;
 	};
 }
diff --git a/Engine/Source/Runtime/SignalProcessing/Private/RuntimeResampler.cpp b/Engine/Source/Runtime/SignalProcessing/Private/RuntimeResampler.cpp
new file mode 100644
--- /dev/null
+++ b/Engine/Source/Runtime/SignalProcessing/Private/RuntimeResampler.cpp
@@ -0,0 +1,647 @@
+// Copyright Epic Games, Inc. All Rights Reserved.
+
+#include "DSP/RuntimeResampler.h"
+
+#include "DSP/BufferVectorOperations.h"
+#include "DSP/Dsp.h"
+#include "DSP/MultichannelBuffer.h"
+#include "HAL/PlatformMath.h"
+#include "Math/UnrealMathUtility.h"
+
+// Throughout this file, all variables with the FP suffix are in the 16.16 fixed point format.
+
+namespace Audio
+{
+	const float FRuntimeResampler::MaxFrameRatio = 100.f;
+	const float FRuntimeResampler::MinFrameRatio = 0.01f;
+	const int32 FRuntimeResampler::FPScale = 65536;
+	const float FRuntimeResampler::FPScaleFloat = 65536;
+
+	FRuntimeResampler::FRuntimeResampler(int32 InNumChannels)
+	{
+		Reset(InNumChannels);
+	}
+
+	void FRuntimeResampler::Reset(const int32 InNumChannels)
+	{
+		check(InNumChannels > 0);
+		CurrentInputFrameIndexFP = 0;
+		CurrentFrameRatioFP = FPScale;
+		TargetFrameRatioFP = FPScale;
+		FrameRatioFrameDeltaFP = 0;
+		NumFramesToInterpolate = 0;
+
+		PreviousFrame.Empty();
+		PreviousFrame.AddZeroed(InNumChannels);
+
+		TempInputPointers.Empty();
+		TempInputPointers.AddZeroed(InNumChannels);
+
+		TempOutputPointers.Empty();
+		TempOutputPointers.AddZeroed(InNumChannels);
+	}
+
+	void FRuntimeResampler::SetFrameRatio(float InRatio, int32 InDesiredNumFramesToInterpolate)
+	{
+		if (ensureMsgf((InRatio >= MinFrameRatio) && (InRatio <= MaxFrameRatio), TEXT("The frame ratio (%f) must be between %f and %f."), InRatio, MinFrameRatio, MaxFrameRatio))
+		{
+			const uint32 RatioFP = FMath::RoundToInt32(InRatio * FPScale);
+
+			if (InDesiredNumFramesToInterpolate <= 0 || RatioFP == CurrentFrameRatioFP)
+			{
+				// Set frame ratio immediately. 
+				CurrentFrameRatioFP = RatioFP;
+				TargetFrameRatioFP = RatioFP;
+				FrameRatioFrameDeltaFP = 0;
+				NumFramesToInterpolate = 0;
+			}
+			else
+			{
+				// Interpolate frame ratio over output frames. 
+				TargetFrameRatioFP = RatioFP;
+				const uint32 RatioChangeAbsolute = uint32(FMath::Abs(int32(TargetFrameRatioFP - CurrentFrameRatioFP)));
+
+				// Round up the frame delta so it's always at least 1
+				const uint32 FrameDeltaAbsolute = FMath::DivideAndRoundUp<uint32>(RatioChangeAbsolute, InDesiredNumFramesToInterpolate);
+
+				// Fudge the number of frames over which we do the interpolation to make sure we land close to the target
+				check(FrameDeltaAbsolute > 0);
+				NumFramesToInterpolate = FMath::DivideAndRoundDown(RatioChangeAbsolute, FrameDeltaAbsolute);
+
+				if (NumFramesToInterpolate > 0)
+				{
+					FrameRatioFrameDeltaFP = int32(FrameDeltaAbsolute);
+					if (TargetFrameRatioFP < CurrentFrameRatioFP)
+					{
+						FrameRatioFrameDeltaFP = -FrameRatioFrameDeltaFP;
+					}
+				}
+				else
+				{
+					CurrentFrameRatioFP = RatioFP;
+					TargetFrameRatioFP = RatioFP;
+					FrameRatioFrameDeltaFP = 0;
+					check(NumFramesToInterpolate == 0);
+				}
+			}
+		}
+	}
+
+	int64 FRuntimeResampler::MapOutputFrameToInputFrameFP(const int32 InOutputFrameIndex) const
+	{
+		checkf(InOutputFrameIndex >= 0, TEXT("Frame index mapping function is only value for outputs frames greater than or equal to 0"));
+
+		int64 InputFrameIndex = 0;
+		if (NumFramesToInterpolate > 0)
+		{
+			if (InOutputFrameIndex < NumFramesToInterpolate)
+			{
+				// Frame ratio interpolation is still occurring at the output frame index.
+				const int64 PreviousIndex = InOutputFrameIndex;
+				const int64 AccumulationOfFrameDeltasFP = FrameRatioFrameDeltaFP * ((PreviousIndex * (PreviousIndex + 1)) / 2);
+				InputFrameIndex = PreviousIndex * CurrentFrameRatioFP + AccumulationOfFrameDeltasFP;
+			}
+			else
+			{
+				// Frame ratio interpolation occurred, but has reached the target frame ratio by the output frame index.
+				const int64 AccumulationOfFrameDeltasFP = FrameRatioFrameDeltaFP * (int64(NumFramesToInterpolate) * (NumFramesToInterpolate + 1) / 2);
+				InputFrameIndex = int64(NumFramesToInterpolate) * CurrentFrameRatioFP + AccumulationOfFrameDeltasFP + int64(InOutputFrameIndex - NumFramesToInterpolate) * TargetFrameRatioFP;
+			}
+		}
+		else
+		{
+			// No interpolation is happening. The math is quite a bit simpler. 
+			InputFrameIndex = int64(TargetFrameRatioFP) * InOutputFrameIndex;
+		}
+
+		// Apply current internal offset. 
+		InputFrameIndex += CurrentInputFrameIndexFP;
+		return InputFrameIndex;
+	}
+
+	int32 FRuntimeResampler::GetNumInputFramesNeededToProduceOutputFrames(const int32 InNumOutputFrames) const
+	{
+		check(InNumOutputFrames >= 0);
+		if (InNumOutputFrames > 0)
+		{
+			const int64 InputFrameFP = MapOutputFrameToInputFrameFP(InNumOutputFrames - 1);
+
+			// We need the sample after InputFrameFP, so we take the floor and add 1. Because we're returning
+			//  the number of samples needed rather than the index of the last sample, we add another 1.
+			return int32(InputFrameFP / FPScale) + 2;
+		}
+		return 0;
+	}
+
+	int32 FRuntimeResampler::GetNumOutputFramesProducedByInputFrames(const int32 InNumInputFrames) const
+	{
+		check(InNumInputFrames >= 0);
+		if (InNumInputFrames <= 0)
+		{
+			return 0;
+		}
+
+		const int64 InputFrameIndexFP = int64(InNumInputFrames - 1) * FPScale;
+		int32 OutputFrameIndex;
+		if (NumFramesToInterpolate > 0)
+		{
+			const int64 InputSampleAtBreakpointFP = MapOutputFrameToInputFrameFP(NumFramesToInterpolate);
+			if (InputFrameIndexFP > InputSampleAtBreakpointFP)
+			{
+				OutputFrameIndex = NumFramesToInterpolate + (InputFrameIndexFP - InputSampleAtBreakpointFP - 1) / TargetFrameRatioFP;
+			}
+			else
+			{
+				// Do a binary search for the previous output sample
+				int32 LowerFrameIndex = 0;
+				int32 UpperFrameIndex = NumFramesToInterpolate;
+				while (UpperFrameIndex > LowerFrameIndex + 1)
+				{
+					const int32 MidFrameIndex = (LowerFrameIndex + UpperFrameIndex) / 2;
+					const int64 InputAtMidOutputFP = MapOutputFrameToInputFrameFP(MidFrameIndex);
+					if (InputAtMidOutputFP < InputFrameIndexFP)
+					{
+						LowerFrameIndex = MidFrameIndex;
+					}
+					else
+					{
+						UpperFrameIndex = MidFrameIndex;
+					}
+				}
+
+				OutputFrameIndex = LowerFrameIndex;
+			}
+		}
+		else
+		{
+			OutputFrameIndex = (InputFrameIndexFP - CurrentInputFrameIndexFP - 1) / TargetFrameRatioFP;
+		}
+
+		return OutputFrameIndex + 1;
+	}
+
+	template<typename OutputBufferType, typename OutputChannelType>
+	int32 FRuntimeResampler::ProcessCircularBufferGeneric(FMultichannelCircularBuffer& InAudio, OutputBufferType& OutAudio)
+	{
+		const int32 NumChannels = PreviousFrame.Num();
+		check(InAudio.Num() == NumChannels && OutAudio.Num() == NumChannels);
+
+		ResamplingParameters Parameters;
+		Parameters.InputAudio = TempInputPointers;
+		Parameters.InputStride = 1;
+		Parameters.OutputAudio = TempOutputPointers;
+		Parameters.OutputStride = 1;
+
+		int32 NumFramesConsumed = 0;
+		int32 NumFramesProduced = 0;
+
+		// The input can be in two disjoint segments, so we run the resampling core up to two times.
+		for (int32 SegmentIdx = 0; SegmentIdx < 2; ++SegmentIdx)
+		{
+			for (int32 ChannelIdx = 0; ChannelIdx < NumChannels; ++ChannelIdx)
+			{
+				// Set up input pointers
+				DisjointedArrayView<const float> InputView = InAudio[ChannelIdx].PeekInPlace(InAudio[ChannelIdx].Num());
+				if (ChannelIdx == 0)
+				{
+					Parameters.NumInputFrames = InputView.FirstBuffer.Num();
+				}
+				else
+				{
+					// We require that all of the circular buffers have the same layout
+					check(Parameters.NumInputFrames == InputView.FirstBuffer.Num());
+				}
+				TempInputPointers[ChannelIdx] = InputView.FirstBuffer.GetData();
+
+				// Set up the output pointers
+				OutputChannelType& OutputChannel = OutAudio[ChannelIdx];
+				TempOutputPointers[ChannelIdx] = OutputChannel.GetData() + NumFramesProduced;
+				if (ChannelIdx == 0)
+				{
+					Parameters.NumOutputFrames = OutputChannel.Num() - NumFramesProduced;
+				}
+				else
+				{
+					check(Parameters.NumOutputFrames == OutputChannel.Num() - NumFramesProduced); // Every output channel must have the same size
+				}
+			}
+
+			// Run the resampler, pop data from the input
+			ProcessAudioInternal(Parameters);
+			NumFramesConsumed += Parameters.OutNumInputFramesConsumed;
+			NumFramesProduced += Parameters.OutNumOutputFramesProduced;
+			for (int32 ChannelIdx = 0; ChannelIdx < NumChannels; ++ChannelIdx)
+			{
+				InAudio[ChannelIdx].Pop(Parameters.OutNumInputFramesConsumed);
+			}
+
+			// Skip the second pass if there's nothing left to do
+			if (NumFramesProduced == OutAudio[0].Num() || InAudio[0].Num() == 0)
+			{
+				break;
+			}
+		}
+
+		return NumFramesProduced;
+	}
+
+	int32 FRuntimeResampler::ProcessCircularBuffer(FMultichannelCircularBuffer& InAudio, FMultichannelBuffer& OutAudio)
+	{
+		return ProcessCircularBufferGeneric<FMultichannelBuffer, FAlignedFloatBuffer>(InAudio, OutAudio);
+	}
+
+	int32 FRuntimeResampler::ProcessCircularBuffer(FMultichannelCircularBuffer& InAudio, FMultichannelBufferView& OutAudio)
+	{
+		return ProcessCircularBufferGeneric<FMultichannelBufferView, TArrayView<float>>(InAudio, OutAudio);
+	}
+
+	void FRuntimeResampler::ProcessInterleaved(const TArrayView<const float> Input, const TArrayView<float> Output, int32& OutNumInputFramesConsumed, int32& OutNumOutputFramesProduced)
+	{
+		const int32 NumChannels = PreviousFrame.Num();
+		for (int32 Channel = 0; Channel < NumChannels; ++Channel)
+		{
+			TempInputPointers[Channel] = Input.GetData() + Channel;
+			TempOutputPointers[Channel] = Output.GetData() + Channel;
+		}
+
+		ResamplingParameters Params{};
+		Params.InputAudio = TempInputPointers;
+		Params.NumInputFrames = Input.Num() / NumChannels;
+		Params.InputStride = NumChannels;
+		Params.OutputAudio = TempOutputPointers;
+		Params.NumOutputFrames = Output.Num() / NumChannels;
+		Params.OutputStride = NumChannels;
+
+		ProcessAudioInternal(Params);
+
+		OutNumInputFramesConsumed = Params.OutNumInputFramesConsumed;
+		OutNumOutputFramesProduced = Params.OutNumOutputFramesProduced;
+	}
+
+	void FRuntimeResampler::ProcessAudioInternal(ResamplingParameters& Parameters)
+	{
+		check(Parameters.InputStride > 0 && Parameters.NumInputFrames >= 0 &&
+			Parameters.OutputStride > 0 && Parameters.NumOutputFrames >= 0);
+
+		const int32 NumChannels = PreviousFrame.Num();
+		check(NumChannels > 0 && NumChannels == Parameters.InputAudio.Num() && NumChannels == Parameters.OutputAudio.Num());
+
+		// Prevent overflowing an int32 in our fixed-point input index by dividing into batches
+		const int32 MaxSamplesPerRun = 30000;
+		if (Parameters.NumInputFrames > 30000)
+		{
+			ProcessAudioInBatches(Parameters, MaxSamplesPerRun);
+			return;
+		}
+
+		// This will be reduced later if we don't have enough input.
+		int32 NumOutputFrames = Parameters.NumOutputFrames;
+
+		if (NumOutputFrames < 1 || Parameters.NumInputFrames < 1)
+		{
+			Parameters.OutNumInputFramesConsumed = 0;
+			Parameters.OutNumOutputFramesProduced = 0;
+			return;
+		}
+
+		float const* const* const InputPointers = Parameters.InputAudio.GetData();
+		float* const* const OutputPointers = Parameters.OutputAudio.GetData();
+
+		uint32 InputFrameRatioFP = CurrentFrameRatioFP;
+		int32 InputFrameIndexFP = CurrentInputFrameIndexFP;
+		const int32 OutputStride = Parameters.OutputStride;
+		const int32 InputStride = Parameters.InputStride;
+
+		// See if we can quickly do Memcpy's instead of resampling
+		if (CurrentInputFrameIndexFP == 0 && TargetFrameRatioFP == FPScale && NumFramesToInterpolate == 0)
+		{
+			if (DoDirectCopy(Parameters))
+			{
+				return;
+			}
+		}
+
+		// Check that InputFrameIndexFP is >= -1
+		checkf(InputFrameIndexFP >= -FPScale, TEXT("Input frame index references discarded data"));
+
+		// Figure out how many samples we can safely generate in the core loop
+		const int32 NumFramesNeeded = GetNumInputFramesNeededToProduceOutputFrames(NumOutputFrames);
+		if (NumFramesNeeded > Parameters.NumInputFrames)
+		{
+			NumOutputFrames = GetNumOutputFramesProducedByInputFrames(Parameters.NumInputFrames);
+		}
+
+		const int32 FinalNumFramesToInterpolate = FMath::Max(NumFramesToInterpolate - NumOutputFrames, 0);
+		const int32 OutputEndIndex = NumOutputFrames * Parameters.OutputStride;
+		int32 OutputSampleIndex = 0;
+
+		// Output any samples that involve PreviousFrame
+		while (InputFrameIndexFP < 0 && OutputSampleIndex < OutputEndIndex)
+		{
+			const float Alpha = (InputFrameIndexFP + FPScale) / FPScaleFloat;
+			for (int32 ChannelIdx = 0; ChannelIdx < NumChannels; ++ChannelIdx)
+			{
+				OutputPointers[ChannelIdx][OutputSampleIndex] = FMath::Lerp(PreviousFrame[ChannelIdx], InputPointers[ChannelIdx][0], Alpha);
+			}
+
+			OutputSampleIndex += OutputStride;
+
+			if (NumFramesToInterpolate > 0)
+			{
+				if (--NumFramesToInterpolate == 0)
+				{
+					InputFrameRatioFP = TargetFrameRatioFP;
+					FrameRatioFrameDeltaFP = 0;
+				}
+				else
+				{
+					InputFrameRatioFP += FrameRatioFrameDeltaFP;
+				}
+			}
+
+			InputFrameIndexFP += InputFrameRatioFP;
+		}
+
+		// Do the bulk of the processing in an optimized core
+		if (NumChannels == 1 && InputStride == 1 && OutputStride == 1)
+		{
+			MonoResamplingCore(InputPointers[0], OutputPointers[0], InputFrameIndexFP, InputFrameRatioFP, OutputSampleIndex, OutputEndIndex);
+		}
+		else if (NumChannels == 2 && InputStride == 2 && OutputStride == 2)
+		{
+			StereoInterleavedResamplingCore(InputPointers[0], OutputPointers[0], InputFrameIndexFP, InputFrameRatioFP, OutputSampleIndex, OutputEndIndex);
+		}
+		else if (NumChannels == 2 && InputStride == 1 && OutputStride == 1)
+		{
+			StereoDeinterleavedResamplingCore(InputPointers[0], InputPointers[1], OutputPointers[0], OutputPointers[1], InputFrameIndexFP, InputFrameRatioFP, OutputSampleIndex, OutputEndIndex);
+		}
+		else
+		{
+			GenericResamplingCore(Parameters, InputFrameIndexFP, InputFrameRatioFP, OutputSampleIndex, OutputEndIndex);
+		}
+
+		NumFramesToInterpolate = FinalNumFramesToInterpolate;
+
+		// Set output values
+		const uint32 NumInputFramesConsumed = InputFrameIndexFP < 0 ? 0 : FMath::Min(InputFrameIndexFP / FPScale + 1, Parameters.NumInputFrames);
+		Parameters.OutNumInputFramesConsumed = NumInputFramesConsumed;
+		Parameters.OutNumOutputFramesProduced = NumOutputFrames;
+
+		// Make sure we don't fail to make any progress 
+		check(NumOutputFrames > 0 || NumInputFramesConsumed > 0);
+
+		// Update saved state for next run
+		if (NumInputFramesConsumed > 0)
+		{
+			InputFrameIndexFP -= NumInputFramesConsumed * FPScale;
+
+			// Save final consumed input frame if we'll need it next run
+			if (InputFrameIndexFP < 0)
+			{
+				const int32 SampleIndexToSave = (NumInputFramesConsumed - 1) * InputStride;
+				for (int32 ChannelIdx = 0; ChannelIdx < NumChannels; ++ChannelIdx)
+				{
+					PreviousFrame[ChannelIdx] = InputPointers[ChannelIdx][SampleIndexToSave];
+				}
+			}
+		}
+
+		if (NumFramesToInterpolate <= 0)
+		{
+			check(NumFramesToInterpolate == 0);
+			CurrentFrameRatioFP = TargetFrameRatioFP;
+			FrameRatioFrameDeltaFP = 0;
+		}
+		else
+		{
+			CurrentFrameRatioFP = InputFrameRatioFP;
+		}
+
+		CurrentInputFrameIndexFP = InputFrameIndexFP;
+	}
+
+	void FRuntimeResampler::ProcessAudioInBatches(ResamplingParameters& Parameters, const int32 MaxSamplesPerBatch)
+	{
+		const int32 NumChannels = PreviousFrame.Num();
+
+		int32 InputFramesRemaining = Parameters.NumInputFrames;
+		Parameters.OutNumInputFramesConsumed = 0;
+		Parameters.OutNumOutputFramesProduced = 0;
+
+		// Create temporary copies of the pointer arrrays and resampling parameters
+		ResamplingParameters BatchParameters(Parameters);
+		TArray<const float*, TInlineAllocator<8>> BatchInputAudio(Parameters.InputAudio);
+		TArray<float*, TInlineAllocator<8>> BatchOutputAudio(Parameters.OutputAudio);
+		BatchParameters.InputAudio = BatchInputAudio;
+		BatchParameters.OutputAudio = BatchOutputAudio;
+
+		while (InputFramesRemaining > 0 && BatchParameters.NumOutputFrames > 0)
+		{
+			// Process a batch
+			BatchParameters.NumInputFrames = FMath::Min(MaxSamplesPerBatch, InputFramesRemaining);
+			ProcessAudioInternal(BatchParameters);
+
+			// Reduce frame counts and advance pointers for the next batch
+			Parameters.OutNumInputFramesConsumed += BatchParameters.OutNumInputFramesConsumed;
+			Parameters.OutNumOutputFramesProduced += BatchParameters.OutNumOutputFramesProduced;
+			InputFramesRemaining -= BatchParameters.OutNumInputFramesConsumed;
+			BatchParameters.NumOutputFrames -= BatchParameters.OutNumOutputFramesProduced;
+			for (int32 ChannelIdx = 0; ChannelIdx < NumChannels; ++ChannelIdx)
+			{
+				BatchInputAudio[ChannelIdx] += BatchParameters.OutNumInputFramesConsumed * Parameters.InputStride;
+				BatchOutputAudio[ChannelIdx] += BatchParameters.OutNumOutputFramesProduced * Parameters.OutputStride;
+			}
+		}
+	}
+
+	bool FRuntimeResampler::DoDirectCopy(ResamplingParameters& Parameters)
+	{
+		const int32 NumChannels = PreviousFrame.Num();
+		float const* const* const InputPointers = Parameters.InputAudio.GetData();
+		float* const* const OutputPointers = Parameters.OutputAudio.GetData();
+
+		if (Parameters.InputStride == 1 && Parameters.OutputStride == 1)
+		{
+			// Copy multi-mono buffers
+			const int32 CopyFrames = FMath::Min(Parameters.NumInputFrames, Parameters.NumOutputFrames);
+			const int32 CopyBytes = CopyFrames * sizeof(float);
+			for (int32 ChannelIdx = 0; ChannelIdx < NumChannels; ++ChannelIdx)
+			{
+				FMemory::Memcpy(OutputPointers[ChannelIdx], InputPointers[ChannelIdx], CopyBytes);
+			}
+			Parameters.OutNumInputFramesConsumed = CopyFrames;
+			Parameters.OutNumOutputFramesProduced = CopyFrames;
+			return true;
+		}
+		else if (Parameters.InputStride == NumChannels && Parameters.OutputStride == NumChannels)
+		{
+			// Check if this is a complete interleaved buffer
+			for (int32 ChannelIdx = 0; ChannelIdx < NumChannels; ++ChannelIdx)
+			{
+				if (InputPointers[ChannelIdx] != InputPointers[0] + ChannelIdx ||
+					OutputPointers[ChannelIdx] != OutputPointers[0] + ChannelIdx)
+				{
+					return false; // Not appropriately interleaved.
+				}
+			}
+
+			// Do one big copy
+			const int32 CopyFrames = FMath::Min(Parameters.NumInputFrames, Parameters.NumOutputFrames);
+			const int32 CopyBytes = CopyFrames * NumChannels * sizeof(float);
+			FMemory::Memcpy(OutputPointers[0], InputPointers[0], CopyBytes);
+			Parameters.OutNumInputFramesConsumed = CopyFrames;
+			Parameters.OutNumOutputFramesProduced = CopyFrames;
+			return true;
+		}
+
+		return false;
+	}
+
+	void FRuntimeResampler::GenericResamplingCore(const ResamplingParameters& Parameters, int32& OutInputFrameIndexFP, uint32& OutInputFrameRatioFP, int32 OutputSampleIndex, const int32 OutputEndIndex)
+	{
+		const int32 NumChannels = PreviousFrame.Num();
+		const int32 OutputStride = Parameters.OutputStride;
+		const int32 InputStride = Parameters.InputStride;
+		float const* const* const InputPointers = Parameters.InputAudio.GetData();
+		float* const* const OutputPointers = Parameters.OutputAudio.GetData();
+		int32 InputFrameIndexFP = OutInputFrameIndexFP;
+		int32 InputFrameRatioFP = OutInputFrameRatioFP;
+
+		// Handle frames where samples are interpolated. 
+		const int32 BatchEndIndex = FMath::Min(OutputEndIndex, OutputSampleIndex + NumFramesToInterpolate * OutputStride);
+		for (; OutputSampleIndex < BatchEndIndex; OutputSampleIndex += OutputStride)
+		{
+			const int32 LowerSampleIndex = (InputFrameIndexFP / FPScale) * InputStride;
+			const int32 UpperSampleIndex = LowerSampleIndex + InputStride;
+			const float Alpha = (InputFrameIndexFP & (FPScale - 1)) / FPScaleFloat;
+			checkSlow(UpperSampleIndex < Parameters.NumInputFrames * InputStride);
+			for (int32 ChannelIdx = 0; ChannelIdx < NumChannels; ++ChannelIdx)
+			{
+				OutputPointers[ChannelIdx][OutputSampleIndex] = FMath::Lerp(InputPointers[ChannelIdx][LowerSampleIndex], InputPointers[ChannelIdx][UpperSampleIndex], Alpha);
+			}
+
+			InputFrameRatioFP += FrameRatioFrameDeltaFP;
+			InputFrameIndexFP += InputFrameRatioFP;
+		}
+
+		// Handle frames where sample rate is constant.
+		for (; OutputSampleIndex < OutputEndIndex; OutputSampleIndex += OutputStride)
+		{
+			const int32 LowerSampleIndex = (InputFrameIndexFP / FPScale) * InputStride;
+			const int32 UpperSampleIndex = LowerSampleIndex + InputStride;
+			const float Alpha = (InputFrameIndexFP & (FPScale - 1)) / FPScaleFloat;
+			checkSlow(UpperSampleIndex < Parameters.NumInputFrames * InputStride);
+			for (int32 ChannelIdx = 0; ChannelIdx < NumChannels; ++ChannelIdx)
+			{
+				OutputPointers[ChannelIdx][OutputSampleIndex] = FMath::Lerp(InputPointers[ChannelIdx][LowerSampleIndex], InputPointers[ChannelIdx][UpperSampleIndex], Alpha);
+			}
+
+			InputFrameIndexFP += TargetFrameRatioFP;
+		}
+
+		OutInputFrameIndexFP = InputFrameIndexFP;
+		OutInputFrameRatioFP = InputFrameRatioFP;
+	}
+
+	void FRuntimeResampler::MonoResamplingCore(const float* Input, float* Output, int32& OutInputFrameIndexFP, uint32& OutInputFrameRatioFP, int32 OutputSampleIndex, const int32 OutputEndIndex)
+	{
+		int32 InputFrameIndexFP = OutInputFrameIndexFP;
+		uint32 InputFrameRatioFP = OutInputFrameRatioFP;
+
+		// Handle frames where samples are interpolated. 
+		const int32 BatchEndIndex = FMath::Min(OutputEndIndex, OutputSampleIndex + NumFramesToInterpolate);
+		for (; OutputSampleIndex < BatchEndIndex; OutputSampleIndex++)
+		{
+			const int32 LowerSampleIndex = InputFrameIndexFP / FPScale;
+			const int32 UpperSampleIndex = LowerSampleIndex + 1;
+			const float Alpha = (InputFrameIndexFP & (FPScale - 1)) / FPScaleFloat;
+			Output[OutputSampleIndex] = FMath::Lerp(Input[LowerSampleIndex], Input[UpperSampleIndex], Alpha);
+
+			InputFrameRatioFP += FrameRatioFrameDeltaFP;
+			InputFrameIndexFP += InputFrameRatioFP;
+		}
+
+		// Handle frames where sample rate is constant.
+		for (; OutputSampleIndex < OutputEndIndex; OutputSampleIndex++)
+		{
+			const int32 LowerSampleIndex = InputFrameIndexFP / FPScale;
+			const int32 UpperSampleIndex = LowerSampleIndex + 1;
+			const float Alpha = (InputFrameIndexFP & (FPScale - 1)) / FPScaleFloat;
+			Output[OutputSampleIndex] = FMath::Lerp(Input[LowerSampleIndex], Input[UpperSampleIndex], Alpha);
+
+			InputFrameIndexFP += TargetFrameRatioFP;
+		}
+
+		OutInputFrameIndexFP = InputFrameIndexFP;
+		OutInputFrameRatioFP = InputFrameRatioFP;
+	}
+
+	void FRuntimeResampler::StereoInterleavedResamplingCore(const float* Input, float* Output, int32& OutInputFrameIndexFP, uint32& OutInputFrameRatioFP, int32 OutputSampleIndex, const int32 OutputEndIndex)
+	{
+		int32 InputFrameIndexFP = OutInputFrameIndexFP;
+		int32 InputFrameRatioFP = OutInputFrameRatioFP;
+
+		// Handle frames where samples are interpolated. 
+		const int32 BatchEndIndex = FMath::Min(OutputEndIndex, OutputSampleIndex + NumFramesToInterpolate * 2);
+		for (; OutputSampleIndex < BatchEndIndex; OutputSampleIndex += 2)
+		{
+			const int32 LowerSampleIndex = (InputFrameIndexFP / FPScale) * 2;
+			const float Alpha = (InputFrameIndexFP & (FPScale - 1)) / FPScaleFloat;
+			Output[OutputSampleIndex] = FMath::Lerp(Input[LowerSampleIndex], Input[LowerSampleIndex + 2], Alpha);
+			Output[OutputSampleIndex + 1] = FMath::Lerp(Input[LowerSampleIndex + 1], Input[LowerSampleIndex + 3], Alpha);
+
+			InputFrameRatioFP += FrameRatioFrameDeltaFP;
+			InputFrameIndexFP += InputFrameRatioFP;
+		}
+
+		// Handle frames where sample rate is constant.
+		for (; OutputSampleIndex < OutputEndIndex; OutputSampleIndex += 2)
+		{
+			const int32 LowerSampleIndex = (InputFrameIndexFP / FPScale) * 2;
+			const float Alpha = (InputFrameIndexFP & (FPScale - 1)) / FPScaleFloat;
+			Output[OutputSampleIndex] = FMath::Lerp(Input[LowerSampleIndex], Input[LowerSampleIndex + 2], Alpha);
+			Output[OutputSampleIndex + 1] = FMath::Lerp(Input[LowerSampleIndex + 1], Input[LowerSampleIndex + 3], Alpha);
+
+			InputFrameIndexFP += TargetFrameRatioFP;
+		}
+
+		OutInputFrameIndexFP = InputFrameIndexFP;
+		OutInputFrameRatioFP = InputFrameRatioFP;
+	}
+
+	void FRuntimeResampler::StereoDeinterleavedResamplingCore(const float* InputLeft, const float* InputRight, float* OutputLeft, float* OutputRight, int32& OutInputFrameIndexFP, uint32& OutInputFrameRatioFP, int32 OutputSampleIndex, const int32 OutputEndIndex)
+	{
+		int32 InputFrameIndexFP = OutInputFrameIndexFP;
+		int32 InputFrameRatioFP = OutInputFrameRatioFP;
+
+		// Handle frames where samples are interpolated. 
+		const int32 BatchEndIndex = FMath::Min(OutputEndIndex, OutputSampleIndex + NumFramesToInterpolate);
+		for (; OutputSampleIndex < BatchEndIndex; OutputSampleIndex++)
+		{
+			const int32 LowerSampleIndex = InputFrameIndexFP / FPScale;
+			const int32 UpperSampleIndex = LowerSampleIndex + 1;
+			const float Alpha = (InputFrameIndexFP & (FPScale - 1)) / FPScaleFloat;
+			OutputLeft[OutputSampleIndex] = FMath::Lerp(InputLeft[LowerSampleIndex], InputLeft[UpperSampleIndex], Alpha);
+			OutputRight[OutputSampleIndex] = FMath::Lerp(InputRight[LowerSampleIndex], InputRight[UpperSampleIndex], Alpha);
+
+			InputFrameRatioFP += FrameRatioFrameDeltaFP;
+			InputFrameIndexFP += InputFrameRatioFP;
+		}
+
+		// Handle frames where sample rate is constant.
+		for (; OutputSampleIndex < OutputEndIndex; OutputSampleIndex++)
+		{
+			const int32 LowerSampleIndex = InputFrameIndexFP / FPScale;
+			const int32 UpperSampleIndex = LowerSampleIndex + 1;
+			const float Alpha = (InputFrameIndexFP & (FPScale - 1)) / FPScaleFloat;
+			OutputLeft[OutputSampleIndex] = FMath::Lerp(InputLeft[LowerSampleIndex], InputLeft[UpperSampleIndex], Alpha);
+			OutputRight[OutputSampleIndex] = FMath::Lerp(InputRight[LowerSampleIndex], InputRight[UpperSampleIndex], Alpha);
+
+			InputFrameIndexFP += TargetFrameRatioFP;
+		}
+
+		OutInputFrameIndexFP = InputFrameIndexFP;
+		OutInputFrameRatioFP = InputFrameRatioFP;
+	}
+}
+
diff --git a/Engine/Source/Runtime/SignalProcessing/Public/DSP/RuntimeResampler.h b/Engine/Source/Runtime/SignalProcessing/Public/DSP/RuntimeResampler.h
new file mode 100644
--- /dev/null
+++ b/Engine/Source/Runtime/SignalProcessing/Public/DSP/RuntimeResampler.h
@@ -0,0 +1,141 @@
+// Copyright Epic Games, Inc. All Rights Reserved.
+
+#pragma once
+
+#include "DSP/BufferVectorOperations.h"
+#include "DSP/Dsp.h"
+#include "DSP/MultichannelBuffer.h"
+
+namespace Audio
+{
+	/** Linear interpolating resampler which supports any number of channels of audio, interleaved or deinterleaved.
+	 To support different types of input and output objects, add a new "Process..." method which wraps around calls to
+	 ProcessAudioInternal(). */
+	class SIGNALPROCESSING_API FRuntimeResampler
+	{
+	public:
+
+		static const float MaxFrameRatio;
+		static const float MinFrameRatio;
+
+		/** Construct a linear resampler.
+		 *
+		 * @param InNumChannel - Number of audio channels in input and output buffers.
+		 */
+		FRuntimeResampler(int32 InNumChannels);
+
+		/** Sets the number of input frames to read per an output frame. 0.5 is
+		 * half speed, 1.f is normal speed, 2.0 is double speed.
+		 *
+		 * @param InRatio - Ratio of input frames consumed per an output frame produced.
+		 * @param InNumFramesToInterpolate - Number of output frames over which
+		 *                                   to interpolate the frame ratio.
+		 */
+		void SetFrameRatio(float InRatio, int32 InNumFramesToInterpolate = 0);
+
+		/** Returns the minimum number of input frames needed to produce the desired
+		 * number of output frames given the current state of the resampler.
+		 *
+		 * @param InNumOutputFrames - The desired number of output frames.
+		 * @return The minimum number of input frames.
+		 */
+		int32 GetNumInputFramesNeededToProduceOutputFrames(int32 InNumOutputFrames) const;
+
+		/** Returns the maximum number of output frames that can be produced
+		 * from the provided number of input frames given the current state of the resampler.
+		 *
+		 * @param InNumInputFrames - The number of provided input frames.
+		 * @return The maximum number of output frames.
+		 */
+		int32 GetNumOutputFramesProducedByInputFrames(int32 InNumInputFrames) const;
+
+		/** Consumes audio from the input buffer and produces audio in the output buffer.
+		 * The desired number of frames to produce is determined by the output audio buffer
+		 * size. For the desired number of samples to be produced, the input audio must have the minimum
+		 * number of frames needed to produce the output frames (see `GetNumInputFramesNeededToProduceOutputFrames(...)`).
+		 * Input samples which are no longer needed are removed from the input buffer.
+		 *
+		 * @param InAudio - Multichannel circular buffer of input audio.
+		 * @param OutAudio - Multichannel buffer of output audio.
+		 * FMultichannelBuffer
+		 *
+		 * @return Actual number of frames produced.
+		 */
+		int32 ProcessCircularBuffer(FMultichannelCircularBuffer& InAudio, FMultichannelBuffer& OutAudio);
+
+		/** Consumes audio from the input buffer and produces audio in the output buffer.
+		 * The desired number of frames to produce is determined by the output audio buffer
+		 * size. For the desired number of samples to be produced, the input audio must have the minimum
+		 * number of frames needed to produce the output frames (see `GetNumInputFramesNeededToProduceOutputFrames(...)`).
+		 * Input samples which are no longer needed are removed from the input buffer.
+		 *
+		 * @param InAudio - Multichannel circular buffer of input audio.
+		 * @param OutAudio - Multichannel buffer of output audio.
+		 * FMultichannelBuffer
+		 *
+		 * @return Actual number of frames produced.
+		 */
+		int32 ProcessCircularBuffer(FMultichannelCircularBuffer& InAudio, FMultichannelBufferView& OutAudio);
+
+		/** Consumes audio in an interleaved channel format from the input buffer and produces audio in the output buffer.
+		 * The desired number of frames to produce is determined by the output audio buffer
+		 * size. For the desired number of samples to be produced, the input audio must have the minimum
+		 * number of frames needed to produce the output frames (see `GetNumInputFramesNeededToProduceOutputFrames(...)`).
+		 *
+		 * @param InAudio - Multichannel interleaved audio to use as input.
+		 * @param OutAudio - Buffer to receive multichannel interleaved output audio.
+		 * @param OutNumInputFramesConsumed - Receives the actual number of frames of input consumed.
+		 * @param OutNumOutputFramesProduced - Receives the actual number of frames of output produced.
+		 */
+		void ProcessInterleaved(TArrayView<const float> Input, TArrayView<float> Output, int32& OutNumInputFramesConsumed, int32& OutNumOutputFramesProduced);
+
+		/** Reset the internal state of the resampler. Call this if the next audio data to be submitted
+		 * is unrelated to the audio that has been submitted previously.
+		*/
+		void Reset(int32 InNumChannels);
+
+	private:
+		int64 MapOutputFrameToInputFrameFP(int32 InOutputFrameIndex) const;
+
+		struct ResamplingParameters
+		{
+			TConstArrayView<const float*> InputAudio;
+			int32 InputStride = 0;
+			int32 NumInputFrames = -1;
+			TConstArrayView<float*> OutputAudio;
+			int32 OutputStride = 0;
+			int32 NumOutputFrames = -1;
+
+			int32 OutNumInputFramesConsumed = 0;
+			int32 OutNumOutputFramesProduced = 0;
+		};
+
+		void ProcessAudioInternal(ResamplingParameters& Parameters);
+
+		void ProcessAudioInBatches(ResamplingParameters& Parameters, int32 MaxSamplesPerBatch);
+
+		/** Perform a quick memcpy if appropriate rather than resampling. Returns true if the memcpy was performed, false otherwise. */
+		bool DoDirectCopy(ResamplingParameters& Parameters);
+
+		void GenericResamplingCore(const ResamplingParameters& Parameters, int32& OutInputFrameIndexFP, uint32& OutInputFrameRatioFP, int32 OutputSampleIndex, int32 OutputEndIndex);
+		void MonoResamplingCore(const float* Input, float* Output, int32& OutInputFrameIndexFP, uint32& OutInputFrameRatioFP, int32 OutputSampleIndex, int32 OutputEndIndex);
+		void StereoInterleavedResamplingCore(const float* Input, float* Output, int32& OutInputFrameIndexFP, uint32& OutInputFrameRatioFP, int32 OutputSampleIndex, int32 OutputEndIndex);
+		void StereoDeinterleavedResamplingCore(const float* InputLeft, const float* InputRight, float* OutputLeft, float* OutputRight, int32& OutInputFrameIndexFP, uint32& OutInputFrameRatioFP, int32 OutputSampleIndex, int32 OutputEndIndex);
+
+		template<typename OutputBufferType, typename OutputChannelType>
+		int32 ProcessCircularBufferGeneric(FMultichannelCircularBuffer& InAudio, OutputBufferType& OutAudio);
+
+		static const int32 FPScale;
+		static const float FPScaleFloat;
+
+		int32 CurrentInputFrameIndexFP = 0;
+		uint32 CurrentFrameRatioFP = 0;
+		uint32 TargetFrameRatioFP = 0;
+		int32 FrameRatioFrameDeltaFP = 0;
+		int32 NumFramesToInterpolate = 0;
+
+		TArray<float, TInlineAllocator<2>> PreviousFrame;
+		TArray<const float*, TInlineAllocator<2>> TempInputPointers;
+		TArray<float*, TInlineAllocator<2>> TempOutputPointers;
+	};
+}
