Minor cleanup in Dsp.h to fix compiler warnings

#rnx

diff --git a/Engine/Source/Runtime/SignalProcessing/Public/DSP/Dsp.h b/Engine/Source/Runtime/SignalProcessing/Public/DSP/Dsp.h
--- a/Engine/Source/Runtime/SignalProcessing/Public/DSP/Dsp.h
+++ b/Engine/Source/Runtime/SignalProcessing/Public/DSP/Dsp.h
@@ -129,17 +129,8 @@ namespace Audio
 		{
 			LastPhasePerSample = -1;
 			LastPhase = InStartingPhaseRadians;
-		}
-
-		FSinOsc2DRotation(FSinOsc2DRotation const& other)
-		{
-			FMemory::Memcpy(this, &other, sizeof(FSinOsc2DRotation));
-		}
-
-		FSinOsc2DRotation& operator=(FSinOsc2DRotation const& other)
-		{
-			FMemory::Memcpy(this, &other, sizeof(FSinOsc2DRotation));
-			return *this;
+			QuadDxVec = VectorZeroFloat();
+			QuadDyVec = VectorZeroFloat();
 		}
 
 		/**
@@ -245,7 +236,6 @@ namespace Audio
 		VectorRegister4Float QuadDxVec, QuadDyVec;
 		float LastPhasePerSample;
 		float LastPhase;
-		float Dx, Dy;
 	}; // end FSinOsc2DRotation
 
 	// Fast tanh based on pade approximation
