Add a buffer overrun test for FRuntimeResampler and fix an edge case that turned up.

#rnx

diff --git a/Engine/Source/Runtime/SignalProcessing/Private/RuntimeResampler.cpp b/Engine/Source/Runtime/SignalProcessing/Private/RuntimeResampler.cpp
--- a/Engine/Source/Runtime/SignalProcessing/Private/RuntimeResampler.cpp
+++ b/Engine/Source/Runtime/SignalProcessing/Private/RuntimeResampler.cpp
@@ -142,8 +142,14 @@ namespace Audio
 		}
 
 		const int64 InputFrameIndexFP = int64(InNumInputFrames - 1) * FPScale;
+		const int64 RelativeInputFrameIndexFP = InputFrameIndexFP - CurrentInputFrameIndexFP;
 		int32 OutputFrameIndex;
-		if (NumFramesToInterpolate > 0)
+		if (RelativeInputFrameIndexFP <= 0)
+		{
+			// We don't have enough input to produce output frame 0.
+			OutputFrameIndex = -1;
+		}
+		else if (NumFramesToInterpolate > 0)
 		{
 			const int64 InputSampleAtBreakpointFP = MapOutputFrameToInputFrameFP(NumFramesToInterpolate);
 			if (InputFrameIndexFP > InputSampleAtBreakpointFP)
@@ -152,7 +158,7 @@ namespace Audio
 			}
 			else
 			{
-				// Do a binary search for the previous output sample
+				// Do a binary search for the previous output frame
 				int32 LowerFrameIndex = 0;
 				int32 UpperFrameIndex = NumFramesToInterpolate;
 				while (UpperFrameIndex > LowerFrameIndex + 1)
@@ -174,7 +180,7 @@ namespace Audio
 		}
 		else
 		{
-			OutputFrameIndex = (InputFrameIndexFP - CurrentInputFrameIndexFP - 1) / TargetFrameRatioFP;
+			OutputFrameIndex = int32((RelativeInputFrameIndexFP - 1) / TargetFrameRatioFP);
 		}
 
 		return OutputFrameIndex + 1;
@@ -288,7 +294,7 @@ namespace Audio
 
 		// Prevent overflowing an int32 in our fixed-point input index by dividing into batches
 		const int32 MaxSamplesPerRun = 30000;
-		if (Parameters.NumInputFrames > 30000)
+		if (Parameters.NumInputFrames > MaxSamplesPerRun)
 		{
 			ProcessAudioInBatches(Parameters, MaxSamplesPerRun);
 			return;
