Support Vorbis without encoding support.

Vorbis without encoding support is ~1/10th the library size and removes
1.6MB of static data for encoder tables. On memory-constrained targets,
disable encoding.

diff --git a/Engine/Source/Runtime/Core/Public/HAL/Platform.h b/Engine/Source/Runtime/Core/Public/HAL/Platform.h
--- a/Engine/Source/Runtime/Core/Public/HAL/Platform.h
+++ b/Engine/Source/Runtime/Core/Public/HAL/Platform.h
@@ -548,6 +548,10 @@
 	#define PLATFORM_SUPPORTS_VORBIS_CODEC						1
 #endif
 
+#ifndef PLATFORM_SUPPORTS_VORBIS_ENCODE
+	#define PLATFORM_SUPPORTS_VORBIS_ENCODE						PLATFORM_SUPPORTS_VORBIS_CODEC
+#endif
+
 #ifndef PLATFORM_USE_MINIMAL_HANG_DETECTION
 	#define PLATFORM_USE_MINIMAL_HANG_DETECTION					0
 #endif
diff --git a/Engine/Source/Runtime/SignalProcessing/Private/FileEncoder.cpp b/Engine/Source/Runtime/SignalProcessing/Private/FileEncoder.cpp
--- a/Engine/Source/Runtime/SignalProcessing/Private/FileEncoder.cpp
+++ b/Engine/Source/Runtime/SignalProcessing/Private/FileEncoder.cpp
@@ -74,7 +74,7 @@ namespace Audio
 		{
 			return new FWavEncoder(QualityInfo, 4096);
 		}
-#if PLATFORM_SUPPORTS_VORBIS_CODEC 
+#if PLATFORM_SUPPORTS_VORBIS_ENCODE
 		if (Extension.Equals(OggExtension))
 		{
 			return new FOggVorbisEncoder(QualityInfo, 4096);
diff --git a/Engine/Source/Runtime/SignalProcessing/Private/OggVorbisEncoder.cpp b/Engine/Source/Runtime/SignalProcessing/Private/OggVorbisEncoder.cpp
--- a/Engine/Source/Runtime/SignalProcessing/Private/OggVorbisEncoder.cpp
+++ b/Engine/Source/Runtime/SignalProcessing/Private/OggVorbisEncoder.cpp
@@ -2,7 +2,7 @@
 
 #include "DSP/Encoders/OggVorbisEncoder.h"
 
-#if PLATFORM_SUPPORTS_VORBIS_CODEC 
+#if PLATFORM_SUPPORTS_VORBIS_ENCODE
 #pragma pack(push, 8)
 #include "vorbis/vorbisenc.h"
 #include "vorbis/vorbisfile.h"
@@ -210,4 +210,5 @@ bool FOggVorbisEncoder::EndFile(TArray<uint8>& OutBytes)
 
 	return true;
 }
-#endif // !PLATFORM_TVOS
+
+#endif // PLATFORM_SUPPORTS_VORBIS_ENCODE
diff --git a/Engine/Source/Runtime/SignalProcessing/Public/DSP/Encoders/OggVorbisEncoder.h b/Engine/Source/Runtime/SignalProcessing/Public/DSP/Encoders/OggVorbisEncoder.h
--- a/Engine/Source/Runtime/SignalProcessing/Public/DSP/Encoders/OggVorbisEncoder.h
+++ b/Engine/Source/Runtime/SignalProcessing/Public/DSP/Encoders/OggVorbisEncoder.h
@@ -5,7 +5,7 @@
 #include "CoreMinimal.h"
 #include "DSP/Encoders/IAudioEncoder.h"
 
-#if !PLATFORM_TVOS
+#if PLATFORM_SUPPORTS_VORBIS_ENCODE
 
 struct FOggVorbisEncoderPrivateState;
 
@@ -34,4 +34,5 @@ private:
 	// This must be a raw pointer because it has a non-default destructor that isn't public.
 	FOggVorbisEncoderPrivateState* PrivateState;
 };
-#endif // !PLATFORM_TVOS
+
+#endif // PLATFORM_SUPPORTS_VORBIS_ENCODE
