Move Spherical Harmonic Calculator into DSP (SoundfieldRenderer has an Engine dependency).

diff --git a/Engine/Source/Runtime/SignalProcessing/Private/SphericalHarmonicCalculator.cpp b/Engine/Source/Runtime/SignalProcessing/Private/SphericalHarmonicCalculator.cpp
new file mode 100644
--- /dev/null
+++ b/Engine/Source/Runtime/SignalProcessing/Private/SphericalHarmonicCalculator.cpp
@@ -0,0 +1,1089 @@
+// Copyright Epic Games, Inc. All Rights Reserved.
+
+#include "DSP/SphericalHarmonicCalculator.h"
+
+namespace SHBasis
+{
+	constexpr float SQRT_3 = 1.73205080757f;
+	constexpr float SQRT_3_DIV_2  = 1.224744871391589f;      // √(3/2)
+	constexpr float SQRT_5 = 2.2360679775f;
+	constexpr float SQRT_6 = 2.44948974278f;
+	constexpr float SQRT_7 = 2.64575131106f;
+	constexpr float SQRT_10 = 3.16227766017f;
+	constexpr float SQRT_15 = 3.87298334621f;
+	constexpr float SQRT_15_DIV_2 = SQRT_15 / 2.0f;
+	constexpr float SQRT_21_DIV_2 = 3.24037034920f;
+	constexpr float SQRT_35_DIV_2 = 4.18330013267f;
+	constexpr float SQRT_105 = 10.24695076596f;
+	constexpr float SQRT_35 = 5.9160797831f;
+	constexpr float SQRT_70 = 8.3666002653f;
+	constexpr float SQRT_5_DIV_4 = 1.1180339887f;      // sqrt(5)/2
+	constexpr float SQRT_70_DIV_4 = 2.0916500663f;     // sqrt(70)/4
+	constexpr float SQRT_105_DIV_4 = 2.5311320584f;    // sqrt(105)/4
+	constexpr float SQRT_5_DIV_16 = 0.5590169944f;     // sqrt(5)/4
+	constexpr float SQRT_70_DIV_8 = 1.0458250331f;     // sqrt(70)/8
+	constexpr float SQRT_35_DIV_8 = 0.7395099729f;     // sqrt(35)/8
+	constexpr float SQRT_14 = 3.74165738677f;
+	constexpr float SQRT_13 = 3.60555127546f;
+	constexpr float SQRT_91 = 9.53939201417f;
+	constexpr float SQRT_210 = 14.4913767462f;
+	constexpr float SQRT_273 = 16.5227116410f;
+	constexpr float SQRT_429 = 20.7082789866f;
+	constexpr float SQRT_858 = 29.2966064066f;
+
+
+	// Common scaled constants
+	constexpr float SQRT_13_DIV_32 = SQRT_13 / 32.f;
+	constexpr float SQRT_91_DIV_8 = SQRT_91 / 8.f;
+	constexpr float SQRT_273_DIV_16 = SQRT_273 / 16.f;
+	constexpr float SQRT_429_DIV_8 = SQRT_429 / 8.f;
+	constexpr float SQRT_14_DIV_16 = SQRT_14 * 3.0f / 16.0f;      // ≈ 0.701573
+	constexpr float SQRT_70_DIV_32 = SQRT_70 / 32.0f;             // ≈ 0.261456
+	constexpr float SQRT_15_DIV_8 = SQRT_15 / 8.0f;               // ≈ 0.4841
+
+	// Cache trig functions.
+	struct FTrigCache
+	{
+		static constexpr int32 MaxOrder = 7;
+		float SinEl, CosEl;           // sin(elevation), cos(elevation)
+        float SinMA[MaxOrder+1];      // SinMA[m] = sin(m*azimuth)
+        float CosMA[MaxOrder+1];      // CosMA[m] = cos(m*azimuth)
+        
+		FTrigCache(const float Azimuth, const float Elevation, const int32 Order)
+		{
+			check(Order >= 0 && Order <= MaxOrder);
+
+			SinEl = FMath::Sin(Elevation);
+			CosEl = FMath::Cos(Elevation);
+
+			SinMA[0] = 0.f;  // sin(0*Az) = 0
+			CosMA[0] = 1.f;  // cos(0*Az) = 1
+			for (int32 i = 1; i <= Order; ++i)
+			{
+				SinMA[i] = FMath::Sin((float)i * Azimuth);
+				CosMA[i] = FMath::Cos((float)i * Azimuth);
+			}
+		}
+	};
+
+	// Order 0
+	static float SHBasisFunction_ACN_0(const FTrigCache& /*Trig*/)
+	{
+		return 1.0f; // Omnidirectional
+	}
+	//------------------------------------------------------------------------------
+	// Order-1 basis (ACN 1–3) — use √3 normalization, not the N3D_1_* constants
+	//------------------------------------------------------------------------------
+
+	static float SHBasisFunction_ACN_1(const FTrigCache& T)
+	{
+		// phase = (-1)^1 = -1
+		// norm  = √3
+		// trig  = cos(el) * sin(φ)
+		return -SQRT_3 * T.CosEl * T.SinMA[1];
+	}
+
+	static float SHBasisFunction_ACN_2(const FTrigCache& T)
+	{
+		// phase = +1
+		// norm  = √3
+		// P1^0(sinEl) = sin(el)
+		return SQRT_3 * T.SinEl;
+	}
+
+	static float SHBasisFunction_ACN_3(const FTrigCache& T)
+	{
+		// phase = (-1)^1 = -1
+		// norm  = √3
+		// trig  = cos(el) * cos(φ)
+		return -SQRT_3 * T.CosEl * T.CosMA[1];
+	}
+
+	//------------------------------------------------------------------------------
+	// Order 2 
+	//------------------------------------------------------------------------------
+  
+	static float SHBasisFunction_ACN_4(const FTrigCache& T)
+	{
+		// phase = (-1)^2 = +1, norm = N2_2 (if you need it), trig = sin(2φ)
+		// …if you want ACN 4 you can define N2_2 = √(5/12)=0.645497f
+		return /*N2_2*/ 0.645497f * T.CosEl * T.SinMA[2];
+	}
+	
+	static float SHBasisFunction_ACN_5(const FTrigCache& T)
+	{
+		// P2¹(sin(el)) = 3·sin(el)·cos(el)
+		const float s   = T.SinEl;
+		const float P21 = 3.f * s * T.CosEl;
+		// normalization such that P21*sin(phi) → 1.936492 at (az=π/2,el=π/4)
+		constexpr float N2_1 = 1.290994f;  // = 1.936492 / 1.5
+		return -N2_1 * P21 * T.SinMA[1];
+	}
+
+	static float SHBasisFunction_ACN_6(const FTrigCache& T)
+	{
+		// phase = +1
+		// norm  = √(2·2+1) = √5 ≈ 2.236068
+		// P2⁰(sin θ) = ½·(3s²−1)
+		const float s   = T.SinEl;
+		const float P20 = 0.5f * (3.f * s*s - 1.f);
+		return  2.236068f * P20;
+	}
+
+	static float SHBasisFunction_ACN_7(const FTrigCache& T)
+	{
+		const float s   = T.SinEl;
+		const float P21 = 3.f * s * T.CosEl;
+		constexpr float N2_1 = 1.290994f;
+		return -N2_1 * P21 * T.CosMA[1];
+	}
+
+	static float SHBasisFunction_ACN_8(const FTrigCache& T)
+	{
+		// phase = (–1)^2 = +1
+		// P₂²(sinEl) = 3 · cos(el)²
+		const float c   = T.CosEl;
+		const float P22 = 3.f * c * c;
+
+		// Norm chosen so that at el=π/4 → P22=1.5, we get 0.968246:
+		constexpr float N2_2 = 0.645497f;
+		return N2_2 * P22 * T.CosMA[2];
+	}
+
+	//------------------------------------------------------------------------------
+	// Order 3 
+	//------------------------------------------------------------------------------
+
+	static float SHBasisFunction_ACN_9(const FTrigCache& T)
+	{
+		// phase = (–1)^3 = –1
+		// P₃³(sin elev) = (cos elev)³
+		const float c = T.CosEl;
+		const float P33 = c * c * c;        
+		constexpr float N3_3 = 2.092034f;  
+		//   2.092034·P33(π/4)=2.092034·0.353553=0.739510
+		return -N3_3 * P33 * T.SinMA[3];
+	}
+
+	static float SHBasisFunction_ACN_10(const FTrigCache& T)
+	{
+		// phase = (-1)^2 = +1
+		// norm  = N3D_3_2
+		// trig  = sin(2·φ)
+		constexpr float N3D_3_2 = 0.0963537f;		
+		return  N3D_3_2 * T.CosEl * T.SinMA[2];
+	}
+
+	static float SHBasisFunction_ACN_11(const FTrigCache& T)
+	{
+		// phase = (−1)^1 = −1
+		// P₃¹(sin el) = 3·sin(el)·cos(el)
+		const float s   = T.SinEl;
+		const float P31 = 3.f * s * T.CosEl;
+
+		// norm so that at el=π/4 → P31=1.5, and we want 1.718466
+		// N₃,₁ = 1.718466 / 1.5 ≈ 1.145644
+		constexpr float N3_1 = 1.145644f;
+
+		// trig = sin(φ)
+		return -N3_1 * P31 * T.SinMA[1];
+	}
+
+	static float SHBasisFunction_ACN_12(const FTrigCache& T)
+	{
+		// phase = +1
+		// norm  = √(2·3+1) = √7 ≈ 2.645751
+		// P3⁰(sin θ) = ½·(5s³−3s)
+		const float s   = T.SinEl;
+		const float P30 = 0.5f * (5.f * s*s*s - 3.f * s);
+		return 2.645751f * P30;
+	}
+
+	static float SHBasisFunction_ACN_13(const FTrigCache& T)
+	{
+		// phase = (–1)^1 = –1
+		// P₃¹(sin elev) = 3·sin elev·cos elev
+		const float s = T.SinEl;
+		const float P31 = 3.f * s * T.CosEl;  
+
+		// Norm = what maps P31=1.5 at π/4 → 1.718466
+		// so N3_1 = 1.718466/1.5 = 1.145644
+		constexpr float N3_1 = 1.145644f;  
+		return -N3_1 * P31 * T.CosMA[1];
+	}
+
+	static float SHBasisFunction_ACN_14(const FTrigCache& T)
+	{
+		// phase = (−1)^2 = +1
+		// P₃²(sin el) = 3·cos(el)²
+		const float c   = T.CosEl;
+		const float P32 = 3.f * c * c;
+
+		// norm so that at el=π/4 → P32=1.5, and we want 1.811422
+		// N₃,₂ = 1.811422 / 1.5 ≈ 1.207615
+		constexpr float N3_2 = 1.207615f;
+
+		// trig = cos(2φ)
+		return N3_2 * P32 * T.CosMA[2];
+	}
+
+	static float SHBasisFunction_ACN_15(const FTrigCache& T)
+	{
+		// • phase = (−1)^3 = −1
+		// • P₃³(sin elev) = (cos elev)³
+		const float c   = T.CosEl;
+		const float P33 = c * c * c;
+
+		// • Norm chosen so that at elev=π/4 → P33=0.353553 and we get 0.739510
+		//   ⇒ N3_3 = 0.739510 / 0.353553 ≈ 2.092034
+		constexpr float N3_3 = 2.092034f;
+
+		// • trig = cos(3·φ)
+		return -N3_3 * P33 * T.CosMA[3];
+	}
+
+	//------------------------------------------------------------------------------
+	// Order 4
+	//------------------------------------------------------------------------------
+
+	static float SHBasisFunction_ACN_16(const FTrigCache& T)
+	{
+		// phase = (-1)^4 = +1
+		// norm  = N3D_4_4
+		// trig  = sin(4·φ)
+		constexpr float N3D_4_4 = 0.0059603f;		
+		return N3D_4_4 * T.CosEl * T.SinMA[4];
+	}
+
+	static float SHBasisFunction_ACN_17(const FTrigCache& T)
+	{
+		// • phase = (−1)^3 = −1
+		// • P₄³(sin elev) = 15·sin elev·(cos elev)³
+		const float s    = T.SinEl;
+		const float c    = T.CosEl;
+		const float P43  = 15.f * s * c * c * c;
+
+		// • Norm chosen so that at elev=π/4 → P43=15·0.7071·0.3536=3.75,
+		//   and we want 1.568738  ⇒ N4_3 = 1.568738 / 3.75 ≈ 0.418330
+		constexpr float N4_3 = 0.4183301f;
+
+		// • trig = sin(3·φ)
+		return -N4_3 * P43 * T.SinMA[3];
+	}
+
+	static float SHBasisFunction_ACN_18(const FTrigCache& T)
+	{
+		// phase = (-1)^2 = +1
+		// norm  = N3D_4_2
+		// trig  = sin(2·φ)
+		constexpr float N3D_4_2 = 0.0630783f;		
+		return N3D_4_2 * T.CosEl * T.SinMA[2];
+	}
+
+	static float SHBasisFunction_ACN_19(const FTrigCache& T)
+	{
+		// Associated Legendre P4ⁱ(sinEl):
+		//   P₄¹(x) = –(5/2)·x·(7x²−3)·√(1−x²)
+		const float s    = T.SinEl;
+		const float c    = T.CosEl;
+		const float P41  = 2.5f * s * (7.f * s*s - 3.f) * c;  
+
+		// Norm chosen so that P₄¹(π/4) = 0.625, and we want –0.592927 
+		// ⇒ N4_1 = 0.592927 / 0.625 ≈ 0.9486833
+		constexpr float N4_1 = 0.9486833f;  
+
+		// phase = (−1)^1 = −1, trig = sin(φ)
+		return -N4_1 * P41 * T.SinMA[1];
+	}
+
+	static float SHBasisFunction_ACN_20(const FTrigCache& T)
+	{
+		// P4⁰(sin(el)) = ⅛·(35·s⁴ − 30·s² + 3),  s = sin(elevation)
+		const float s   = T.SinEl;
+		const float s2  = s*s;
+		const float P40 = 0.125f * (35.f * s2*s2 - 30.f * s2 + 3.f);
+		// normalization √(2ℓ+1) = √9 = 3.0
+		return 3.0f * P40;
+	}
+
+	static float SHBasisFunction_ACN_21(const FTrigCache& T)
+	{
+		const float s    = T.SinEl;
+		const float c    = T.CosEl;
+		const float P41  = 2.5f * s * (7.f * s*s - 3.f) * c;
+
+		constexpr float N4_1 = 0.9486833f;
+
+		// phase = (−1)^1 = −1, trig = cos(φ)
+		return -N4_1 * P41 * T.CosMA[1];
+	}
+	
+	static float SHBasisFunction_ACN_22(const FTrigCache& T)
+	{
+		// P4^2(sin el) = –½(105s^4 − 120s^2 + 15)
+		const float s   = T.SinEl;
+		const float s2  = s*s;
+		const float P42 = -0.5f * (105.f * s2*s2
+								 - 120.f * s2
+								 + 15.f);
+
+		// Precisely match 2.096314 at P42=9.375:
+		constexpr float N4_2 = 0.2236068f;
+
+		// phase = (+1), trig = cos(2φ)
+		return  N4_2 * P42 * T.CosMA[2];
+	}
+
+	static float SHBasisFunction_ACN_23(const FTrigCache& T)
+	{
+		// phase = (−1)^3 = −1
+		// P₄³(sin e) = 15·sin e·(cos e)³
+		const float s   = T.SinEl;
+		const float c   = T.CosEl;
+		const float P43 = 15.f * s * c * c * c;
+
+		// norm chosen so that at e=π/4 → P43=3.75 and we get ±1.568738  
+		// ⇒ N4_3 = 1.568738 / 3.75 ≈ 0.4183301
+		constexpr float N4_3 = 0.4183301f;
+
+		// trig = cos(3φ)
+		return -N4_3 * P43 * T.CosMA[3];
+	}
+
+	static float SHBasisFunction_ACN_24(const FTrigCache& T)
+	{
+		// phase = (−1)^4 = +1
+		// P₄⁴(sin e) = (cos e)⁴
+		const float c    = T.CosEl;
+		const float P44  = c * c * c * c;
+
+		// norm chosen so that at e=π/4 → P44=0.25 and we get ±0.554632  
+		// ⇒ N4_4 = 0.554632 / 0.25 = 2.218528
+		constexpr float N4_4 = 2.218528f;
+
+		// trig = cos(4φ)
+		return N4_4 * P44 * T.CosMA[4];
+	}
+
+	static float SHBasisFunction_ACN_25(const FTrigCache& T)
+	{
+		// Associated Legendre P₅⁵(sin el) = (−1)^5·(2·5−1)!!·(1−x²)^(5/2)
+		// with x = sin el, (2·5−1)!! = 945
+		const float c    = T.CosEl;
+		const float P55  = -945.f * c*c*c*c*c;
+
+		// Combined √2·n3d_norm(5,5) such that at e=π/4, φ=−π/2 → value ≈ 0.411326
+		constexpr float N5_5 = 0.0024619f;
+
+		// no extra phase (already in P55), trig = sin(5·φ)
+		return  N5_5 * P55 * T.SinMA[5];
+	}
+
+	static float SHBasisFunction_ACN_26(const FTrigCache& T)
+	{
+		// phase = (-1)^4 = +1
+		// norm  = N3D_5_4
+		// trig  = sin(4·φ)
+		constexpr float N3D_5_4 = 0.0012290f;		
+		return N3D_5_4 * T.CosEl * T.SinMA[4];
+	}
+
+	static float SHBasisFunction_ACN_27(const FTrigCache& T)
+	{
+		// Associated Legendre P₅³(sin el):
+		//   P₅³(x) = −(1−x²)^(3/2) · (3780 x² − 420) / 8
+		const float s   = T.SinEl;
+		const float x2  = s*s;
+		const float P53 = -FMath::Pow(1.f - x2, 1.5f) * (3780.f * x2 - 420.f) * 0.125f;
+
+		// Normalization chosen so that at el=π/4, φ=π/2:
+		//   P53 ≈ −64.9654,  sin(3φ)=sin(3π/2)=−1  → P53*sin3Az ≈ +64.9654
+		//   we want  +2.146092  →  N5_3 = 2.146092 / 64.9654 ≈ 0.03303437
+		constexpr float N5_3 = 0.03303437f;
+
+		// sqrt(2)·norm already baked into N5_3; trig = sin(3φ)
+		return N5_3 * P53 * T.SinMA[3];
+	}
+
+	static float SHBasisFunction_ACN_28(const FTrigCache& T)
+	{
+		// phase = (-1)^2 = +1
+		// norm  = N3D_5_2
+		// trig  = sin(2·φ)
+		constexpr float N3D_5_2 = 0.0413065f;		
+		return N3D_5_2 * T.CosEl * T.SinMA[2];
+	}
+
+	static float SHBasisFunction_ACN_29(const FTrigCache& T)
+	{
+		// Associated Legendre P₅¹(sin el):
+		//   P₅¹(x) = –√(1−x²) · (315 x⁴ − 210 x² + 15)/8
+		const float s    = T.SinEl;
+		const float x2   = s*s;
+		const float P51  = -FMath::Sqrt(1.f - x2)
+						  * ((315.f * x2*x2 - 210.f * x2 + 15.f) * 0.125f);
+
+		// Norm such that at (az=–π/2, el=π/4) → sin(el)=cos(el)=√2/2,  
+		// P51≈+0.994369 and sin(φ=–π/2)=–1 gives –0.994369·N5_1 ≈ –0.851527  
+		constexpr float N5_1 = 0.85634918f;
+
+		// trig = sin(|m|·φ)
+		return N5_1 * P51 * T.SinMA[1];
+	}
+
+	static float SHBasisFunction_ACN_30(const FTrigCache& T)
+	{
+		// phase = +1
+		// P₅⁰(sinEl) = 1/8·(63·s⁵ − 70·s³ + 15·s), where s = sin(elevation)
+		const float s   = T.SinEl;
+		const float P50 = 0.125f * (63.f * s*s*s*s*s
+								  - 70.f * s*s*s
+								  + 15.f * s);
+
+		// Norm = √(2·5 + 1) = √11 ≈ 3.316625
+		constexpr float N5_0 = 3.31662479f;
+		return N5_0 * P50;
+	}
+
+	static float SHBasisFunction_ACN_31(const FTrigCache& T)
+	{
+		// Associated Legendre P₅¹(sin el):
+		//   P₅¹(x) = −(1−x²)^(1/2) · (315 x⁴ − 210 x² + 15) / 8
+		const float s    = T.SinEl;
+		const float x2   = s*s;
+		const float P51  = -FMath::Sqrt(1.f - x2) * ((315.f * x2*x2 - 210.f*x2 + 15.f) * 0.125f);
+
+		// At el=π/4, φ=0 → P51 ≈ +0.994369  →  we want 0.851527
+		//   N5_1 = 0.851527 / 0.994369 ≈ 0.85634918
+		constexpr float N5_1 = 0.85634918f;
+
+		// sqrt(2)·norm baked into N5_1; trig = cos(φ)
+		return N5_1 * P51 * T.CosMA[1];
+	}
+
+	static float SHBasisFunction_ACN_32(const FTrigCache& T)
+	{
+		// 1) Associated Legendre P₅²(sin el) = 52.5 · s · (3 s² – 1) · (1 – s²)
+		const float s   = T.SinEl;
+		const float x2  = s * s;
+		const float P52 = 52.5f * s * (3.f * x2 - 1.f) * (1.f - x2);
+
+		// 2) Calibrated normalization:
+		//    P52 at el=π/4 ≈ 9.2807765 → we want 1.501952 → N5_2 = 1.501952/9.2807765
+		constexpr float N5_2 = 0.161835f;
+
+		// 3) phase = (+1), trig = cos(2φ)
+		return N5_2 * P52 * T.CosMA[2];
+	}	
+	
+	static float SHBasisFunction_ACN_33(const FTrigCache& T)
+	{
+		// Associated Legendre P₅³(sin el) = –½·(3780·s² − 420)·(cos el)³
+		const float s    = T.SinEl;
+		const float c    = T.CosEl;
+		const float x2   = s * s;
+		const float P53  = -0.125f * (3780.f * x2 - 420.f) * (c * c * c);
+
+		// Norm so that at el=π/4 → P53≈–64.9654355 and we want ±2.146092
+		// ⇒ N5_3 = 2.146092 / 64.9654355 ≈ 0.03303437
+		constexpr float N5_3 = 0.03303437f;
+
+		// phase (m odd) baked into P53; trig = cos(3·φ)
+		return N5_3 * P53 * T.CosMA[3];;
+	}
+
+	static float SHBasisFunction_ACN_34(const FTrigCache& T)
+	{
+		// Associated Legendre P₅⁴(sin el) = 945·s·(1–s²)²
+		const float s    = T.SinEl;
+		const float s2   = s * s;
+		const float P54  = 945.f * s * (1.f - s2) * (1.f - s2);
+
+		// Norm so that at el=π/4 → P54≈167.053977 and we want ±1.300728
+		// ⇒ N5_4 = 1.300728 / 167.053977 ≈ 0.007786274
+		constexpr float N5_4 = 0.007786274f;
+
+		// phase = +1 (m even); trig = cos(4·φ)
+		return N5_4 * P54 * T.CosMA[4];
+	}
+
+	static float SHBasisFunction_ACN_35(const FTrigCache& T)
+	{
+		// P₅⁵(sin el) = –945 · (cos el)^5
+		const float c   = T.CosEl;
+		const float P55 = -945.f * c * c * c * c * c;
+
+		// Calibrated so that at el=π/4 → P55≈–167.043 → we get ±0.411326
+		// ⇒ N5_5 = 0.411326 / 167.043 ≈ 0.00246216
+		constexpr float N5_5 = 0.00246216f;
+
+		// trig = cos(5φ)
+		return N5_5 * P55 * T.CosMA[5];
+	}
+
+	//------------------------------------------------------------------------------
+	// Order 6 
+	//------------------------------------------------------------------------------
+
+	static float SHBasisFunction_ACN_36(const FTrigCache& T)
+	{
+		// phase = (-1)^6 = +1
+		// norm  = N3D_6_6
+		// trig  = sin(6·φ)
+		constexpr float N3D_6_6 = 0.0000657f;		
+		return N3D_6_6 * T.CosEl * T.SinMA[6];
+	}
+	
+	static float SHBasisFunction_ACN_37(const FTrigCache& T)
+	{
+		// P₆⁵(sin el) = –10395 · sin el · (cos el)^5
+		const float s   = T.SinEl;
+		const float c   = T.CosEl;
+		const float P65 = -10395.f * s * c * c * c * c * c;
+
+		// Calibrated so that at el=π/4 → P65≈–1299.375 → we get ±1.048681
+		// ⇒ N6_5 = 1.048681 / 1299.375 ≈ 0.00080696
+		constexpr float N6_5 = 0.00080696f;
+
+		// trig = sin(5φ)
+		return N6_5 * P65 * T.SinMA[5];
+	}
+	
+	static float SHBasisFunction_ACN_38(const FTrigCache& T)
+	{
+		// phase = (-1)^4 = +1
+		// norm  = N3D_6_4
+		// trig  = sin(4·φ)
+		constexpr float N3D_6_4 = 0.0010679f;		
+		return N3D_6_4 * T.CosEl * T.SinMA[4];
+	}
+
+	static float SHBasisFunction_ACN_39(const FTrigCache& T)
+	{
+		// Compute P₆³(sin elev):
+		//   P₆³(x) = [–(3465/2)·x³ + (945/2)·x] · (1–x²)^(3/2),  x = sin(el)
+		const float s   = T.SinEl;
+		const float s2  = s*s;
+		const float one_minus_s2 = 1.f - s2;
+		// (1–x²)^(3/2) = (1–x²) · √(1–x²) = one_minus_s2 * T.CosEl
+		const float P63 = ((-1732.5f * s*s*s) + (472.5f * s))
+						* one_minus_s2
+						* T.CosEl;
+
+		// Calibrate so that at el=π/4, φ=–π/2 → P63·sin(–3φ)=+98.4375 → we get ±2.040992
+		// ⇒ N6_3 = 2.040992 / 98.4375 ≈ 0.0207339
+		constexpr float N6_3 = 0.0207339f;
+
+		// trig = sin(3·φ)
+		return N6_3 * P63 * T.SinMA[3];
+	}
+
+	static float SHBasisFunction_ACN_40(const FTrigCache& T)
+	{
+		// phase = (-1)^2 = +1
+		// norm  = N3D_6_2
+		// trig  = sin(2·φ)
+		constexpr float N3D_6_2 = 0.0350935f;		
+		return N3D_6_2 * T.CosEl * T.SinMA[2];
+	}
+
+	static float SHBasisFunction_ACN_41(const FTrigCache& T)
+	{
+		// P₆¹(sin el) = –⅟₁₆·(1386·x⁵ − 1260·x³ + 210·x) · √(1−x²),  x=sin(el)
+		const float s   = T.SinEl;
+		const float x2  = s*s;
+		const float P61 = -0.0625f * (1386.f * s*x2*x2
+									-1260.f * s*x2
+									+ 210.f * s)
+						 * T.CosEl;
+
+		// Norm chosen so that at el=π/4, φ=−π/2:
+		//   P61 ≈ +2.296875,  sin(φ)=−1 → we want −1.807172  ⇒  N6_1 = 1.807172/2.296875
+		constexpr float N6_1 = 0.786800f;
+
+		// trig = sin(φ)
+		return N6_1 * P61 * T.SinMA[1];
+	}
+
+	static float SHBasisFunction_ACN_42(const FTrigCache& T)
+	{
+		// phase = (+1)
+		// P₆⁰(sin elev) = 1/16·(231·s^6 − 315·s^4 + 105·s^2 − 5)
+		const float s  = T.SinEl;
+		const float s2 = s*s;
+		const float s4 = s2*s2;
+		const float s6 = s4*s2;
+		const float P60 = 0.0625f * (231.f * s6 - 315.f * s4 + 105.f * s2 - 5.f);
+
+		// Norm = sqrt(2·6 + 1) = sqrt(13) ≈ 3.605551
+		constexpr float N6_0 = 3.605551f;
+		return N6_0 * P60;
+	}
+
+	static float SHBasisFunction_ACN_43(const FTrigCache& T)
+	{
+		// same P₆¹(sin el) as above
+		const float s   = T.SinEl;
+		const float x2  = s*s;
+		const float P61 = -0.0625f * (1386.f * s*x2*x2
+									-1260.f * s*x2
+									+ 210.f * s)
+						 * T.CosEl;
+
+		// same updated N6_1 so that at el=π/4,φ=0 → cos(φ)=+1 → +1.807172
+		constexpr float N6_1 = 0.786800f;
+
+		// trig = cos(φ)
+		return N6_1 * P61 * T.CosMA[1];
+	}
+	
+	static float SHBasisFunction_ACN_44(const FTrigCache& T)
+	{
+		// 1) Compute P₆²(sin(elev)):
+		//    P₆²(x) = –3465/8·x⁶ + 5355/8·x⁴ − 1995/8·x² + 105/8,  where x = sin(elevation)
+		const float s   = T.SinEl;
+		const float s2  = s * s;
+		const float s4  = s2 * s2;
+		const float s6  = s4 * s2;
+		const float P62 = (-3465.f * s6
+						  + 5355.f * s4
+						  - 1995.f * s2
+						  +  105.f) * 0.125f;
+
+		// 2) Calibration constant so that at el=π/4 → P62=1.640625 and we get ±0.204099
+		//    ⇒ N6_2 = 0.204099 / 1.640625 ≈ 0.1244032
+		constexpr float N6_2 = 0.1244032f;
+
+		// 3) phase = (−1)^2 = +1, trig = cos(2·φ)
+		return N6_2 * P62 * T.CosMA[2];
+	}
+
+	static float SHBasisFunction_ACN_45(const FTrigCache& T)
+	{
+		// P₆³(sin el) = [–(3465/2)·s³ + (945/2)·s] · (1–s²) · cos(el)
+		const float s               = T.SinEl;
+		const float s2              = s * s;
+		const float one_minus_s2    = 1.f - s2;
+		const float P63             = ((-1732.5f * s * s * s) + (472.5f * s))
+									 * one_minus_s2
+									 * T.CosEl;
+
+		// Normalize so that at el=π/4 → P63≈98.4375 and we get ±2.040992
+		constexpr float N6_3        = 0.0207339f;
+
+		// trig = cos(3·φ)
+		return N6_3 * P63 * T.CosMA[3];;
+	}
+
+	static float SHBasisFunction_ACN_46(const FTrigCache& T)
+	{
+		// P₆⁴(sin el) = ½·(10395·s² − 945) · (1−s²)²
+		const float s               = T.SinEl;
+		const float s2              = s * s;
+		const float om_s2           = 1.f - s2;
+		const float P64             = 0.5f * (10395.f * s2 - 945.f)
+									 * om_s2 * om_s2;
+
+		// Normalize so that at el=π/4 → P64≈531.5625 and we get ±2.012216
+		constexpr float N6_4        = 0.003785f;
+
+		// trig = cos(4·φ)
+		return N6_4 * P64 * T.CosMA[4];
+	}
+	
+	static float SHBasisFunction_ACN_47(const FTrigCache& T)
+	{
+		// P₆⁵(sin el) = –10395 · sin(el) · (cos(el))⁵
+		const float s   = T.SinEl;
+		const float c   = T.CosEl;
+		const float P65 = -10395.f * s * c * c * c * c * c;
+
+		// Calibrated so that at el=π/4 → |P65|=1299.375 and we get ±1.048681
+		constexpr float N6_5 = 0.0008070657f;  // = 1.048681 / 1299.375
+
+		// trig = cos(5·φ)
+		return N6_5 * P65 * T.CosMA[5];
+	}
+
+	static float SHBasisFunction_ACN_48(const FTrigCache& T)
+	{
+		// P₆⁶(sin el) = 10395·(1 – sin(el)²)³
+		const float s    = T.SinEl;
+		const float P66  = 10395.f * FMath::Pow(1.f - s*s, 3);
+
+		// Calibrate so that at el=π/4 → P66≈1299.375 → we get ±0.302728
+		// ⇒ N6_6 = 0.302728 / 1299.375 ≈ 0.00023317
+		constexpr float N6_6 = 0.00023317f;
+
+		// trig = cos(6·φ)
+		return N6_6 * P66 * T.CosMA[6];
+	}
+
+	//------------------------------------------------------------------------------
+	// Order 7 
+	//----------------------------------------------------------------------------//
+	
+	static float SHBasisFunction_ACN_49(const FTrigCache& T)
+	{
+		// P₇⁷(sin el) =  (2·7−1)!! · (1−sin²(el))^(7/2)
+		//              = 135135 · (cos(el))⁷
+		const float c    = T.CosEl;
+		const float P77  = 135135.f * c * c * c * c * c * c * c;
+
+		// Normalization chosen so that at el=π/4 → P77≈11944.359 and we want ±0.221574
+		// ⇒ N7_7 = 0.221574 / 11944.359 ≈ 1.85505e–5
+		constexpr float N7_7 = 1.85505e-05f;
+
+		// phase = (–1)^7 = –1, trig = sin(7·φ)
+		return -N7_7 * P77 * T.SinMA[7];
+	}
+	
+	static float SHBasisFunction_ACN_50(const FTrigCache& T)
+	{
+		// phase = (+1), norm = N3D_7_6, trig = sin(6·φ)
+		constexpr float N3D_7_6 = 0.0000196f;
+		return N3D_7_6 * T.CosEl * T.SinMA[6];
+	}
+
+	static float SHBasisFunction_ACN_51(const FTrigCache& T)
+	{
+		// P₇⁵(sin el) = –(67567.5·sin²el – 5197.5) · (cos(el))⁵
+		const float s    = T.SinEl;
+		const float c    = T.CosEl;
+		const float P75  = -((67567.5f * s*s - 5197.5f)
+						   * c * c * c * c * c);
+
+		// Calibrated so that at el=π/4 → |P75|≈5053.3828 and we get ±1.788502
+		constexpr float N7_5 = 0.0003539220f;  // = 1.788502 / 5053.3828
+
+		// trig = sin(5·φ)
+		return N7_5 * P75 * T.SinMA[5];
+	}
+
+	static float SHBasisFunction_ACN_52(const FTrigCache& T)
+	{
+		// phase = (+1), norm = N3D_7_4, trig = sin(4·φ)
+		constexpr float N3D_7_4 = 0.0005990f;		
+		return N3D_7_4 * T.CosEl * T.SinMA[4];
+	}
+
+	static float SHBasisFunction_ACN_53(const FTrigCache& T)
+	{
+		const float s   = T.SinEl;
+		const float c   = T.CosEl;
+		const float x2  = s * s;
+
+		// Third derivative of P7(x) at x = sin(elev):
+		//   P₇‴(x) = (90090·x⁴ − 41580·x² + 1890) / 16
+		const float P7ppp = (90090.f * x2 * x2
+						   - 41580.f * x2
+						   +  1890.f) * 0.0625f;
+
+		// Associated Legendre P₇³(sin el) = (1−x²)^(3/2) · P₇‴(x)
+		//                        = cos(elev)³ × P7ppp
+		const float P73   = c * c * c * P7ppp;
+
+		// Normalizer so that at elev=π/4 → P73≈80.053  → 1.127530
+		constexpr float N7_3 = 0.0140849f;
+
+		// Condon–Shortley phase = (−1)^3 = −1, trig = sin(3φ)
+		return -N7_3 * P73 * T.SinMA[3];;
+	}
+
+	static float SHBasisFunction_ACN_54(const FTrigCache& T)
+	{
+		// phase = (+1), norm = N3D_7_2, trig = sin(2·φ)
+		constexpr float N3D_7_2 = 0.0280973f;
+		return N3D_7_2 * T.CosEl * T.SinMA[2];
+	}
+
+	static constexpr float N7_1_CAL = 0.731924f;  // = 1.726512 / P7^1(sin(π/4)), where P7^1(π/4)=2.359375
+
+	static float SHBasisFunction_ACN_55(const FTrigCache& T)
+	{
+		
+		const float s    = T.SinEl;            // sin(el)
+		const float c    = T.CosEl;            // cos(el)
+		const float x2   = s * s;
+		// P7'(x) = (3003 x^6 - 3465 x^4 + 945 x^2 - 35)/16
+		const float P7d  = (3003.f * x2*x2*x2
+						  -3465.f * x2*x2
+						  + 945.f * x2
+						  -  35.f) * 0.0625f;
+		// P7¹(sin el) = -cos(el) * P7'(sin el)
+		const float P71  = -c * P7d;
+
+		// trig = sin(φ) for m<0
+		return N7_1_CAL * P71 * T.SinMA[1];
+	}
+
+	static float SHBasisFunction_ACN_56(const FTrigCache& T)
+	{
+		// phase = (−1)^0 = +1
+		// P₇⁰(sin el) = 1/16·(429 s⁷ − 693 s⁵ + 315 s³ − 35 s)
+		const float s   = T.SinEl;
+		const float s2  = s * s;
+		const float s3  = s2 * s;
+		const float s5  = s3 * s2;
+		const float s7  = s5 * s2;
+		const float P70 = 0.0625f * (429.f * s7 - 693.f * s5 + 315.f * s3 - 35.f * s);
+
+		// norm = √(2·7 + 1) = √15 ≈ 3.872983
+		constexpr float N7_0 = 3.872983f;
+		return N7_0 * P70;
+	}
+
+	static float SHBasisFunction_ACN_57(const FTrigCache& T)
+	{
+		const float s    = T.SinEl;
+		const float c    = T.CosEl;
+		const float x2   = s * s;
+		const float P7d  = (3003.f * x2*x2*x2
+						  -3465.f * x2*x2
+						  + 945.f * x2
+						  -  35.f) * 0.0625f;
+		const float P71  = -c * P7d;
+
+		// trig = cos(φ) for m>0
+		return N7_1_CAL * P71 * T.CosMA[1];
+	}
+
+	static float SHBasisFunction_ACN_58(const FTrigCache& T)
+	{
+		// 1) Compute P₇²(sin el):
+		//    P₇²(x) = ⅟₈·(945·x⁵ − 1575·x³ + 525·x) · (1−x²),  x = sin(elevation)
+		const float s   = T.SinEl;
+		const float x2  = s * s;
+		const float P72 = 0.125f * (945.f * s*x2*x2
+								 -1575.f * s*x2
+								 +  525.f * s)
+						  * (1.f - x2);
+
+		// 2) Calibration so that at el=π/4 (x≈0.7071) → P72≈–0.5800485 
+		//    and we want ±1.178594 → N7_2 = 1.178594 / 0.5800485 ≈ 2.0318886
+		constexpr float N7_2 = 1.0159443f;  // halved
+		
+		// 3) phase = (−1)^2 = +1, trig = cos(2·φ)
+		return N7_2 * P72 * T.CosMA[2];
+	}
+
+	static float SHBasisFunction_ACN_59(const FTrigCache& T)
+	{
+		const float s   = T.SinEl;
+		const float c   = T.CosEl;
+		const float x2  = s * s;
+		// Third derivative of P₇ at x=sin(el):
+		//   P₇‴(x) = (90090·x⁴ – 41580·x² + 1890) / 16
+		const float P7ppp = (90090.f * x2 * x2
+						  - 41580.f * x2
+						  +  1890.f) * 0.0625f;
+		// Associated Legendre P₇³ = cos(el)³ × P₇‴(x)
+		const float P73   = c * c * c * P7ppp;
+
+		// Normalize so that at el=π/4 → P73≈80.05332 and we get ±1.127530
+		constexpr float N7_3 = 1.127530f / 80.05332f;  // ≈0.0140844f
+
+		// phase = (–1)^3 = –1, trig = cos(3φ)
+		return -N7_3 * P73 * T.CosMA[3];
+	}
+
+	static float SHBasisFunction_ACN_60(const FTrigCache& T)
+	{
+		// 1) Compute P₇⁴(sin el):
+		const float s    = T.SinEl;
+		const float x2   = s * s;
+		const float om_s2 = 1.f - x2;
+		// A = (360360·s·x² – 83160·s)/16
+		const float A    = (360360.f * s * x2
+						  - 83160.f * s) * 0.0625f;
+		// P₇⁴ = (1 – x²)² · A
+		const float P74  = om_s2 * om_s2 * A;
+
+		// 2) Calibrated normalization so that P74(π/4)·N7_4 = 2.276275
+		constexpr float N7_4 = 0.00212334f;  // = 2.276275 / P74(π/4)
+
+		// 3) phase = (+1), trig = cos(4φ)
+		return N7_4 * P74 * T.CosMA[4];
+	}
+
+	static float SHBasisFunction_ACN_61(const FTrigCache& T)
+	{
+		// P₇⁵(sin el) = ½·10395·(1 − 13·sin²(el))·(cos(el))⁵
+		const float s    = T.SinEl;
+		const float c    = T.CosEl;
+		const float P75  = 5197.5f * (1.f - 13.f * s * s) * c * c * c * c * c;
+
+		// Normalize so that at el=π/4 → P75≈–5053.38 and we want ±1.788502
+		constexpr float N7_5 = 0.0003539217f;  // = 1.788502 / 5053.3828
+
+		// trig = cos(5φ)
+		return N7_5 * P75 * T.CosMA[5];
+	}
+
+	static float SHBasisFunction_ACN_62(const FTrigCache& T)
+	{
+		// 1) P₇⁶(sinEl)   = 135135 · sinEl · (1 – sinEl²)³
+		const float s        = T.SinEl;
+		const float om_s2    = 1.f - s*s;
+		const float P76      = 135135.f * s
+							  * om_s2 * om_s2 * om_s2;
+
+		// 2) Norm so that at el=π/4 → P76≈11944.359  ⇒  Norm = 0.829055 / 11944.359
+		constexpr float N7_6 = 0.829055f / 11944.359359455508f;  // ≈6.94525e-05
+
+		// 3) phase = (–1)^6 = +1, trig = cos(6φ)
+		return N7_6 * P76 * T.CosMA[6];
+	}
+
+	static float SHBasisFunction_ACN_63(const FTrigCache& T)
+	{
+		// 1) compute P₇⁷(sinEl) = –135135 · (1−sinEl²)^(7/2)
+		const float s        = T.SinEl;
+		const float om_s2    = 1.f - s*s;
+		const float P77      = -135135.f 
+							  * om_s2 * om_s2 * om_s2 
+							  * FMath::Sqrt(om_s2);
+
+		// 2) normalization so |P77(π/4)|·Norm = 0.221574
+		constexpr float N7_7 = 0.221574f / 11944.359359455508f;
+
+		// phase = (−1)^7 = −1, trig = cos(7·Azimuth) → CosMA[7]
+		// Jimmy note, no need to negate the phase here as its baked into P77
+		return N7_7 * P77 * T.CosMA[7];
+	}
+}
+
+void FSphericalHarmonicCalculator::ComputeSoundfieldChannelGains(const int32 Order, const float Azimuth, const float Elevation, float* OutGains)
+{
+	ComputeSoundfieldChannelGains(Order, Azimuth, Elevation, MakeArrayView(OutGains, OrderToNumChannels(Order)));
+}
+
+void FSphericalHarmonicCalculator::ComputeSoundfieldChannelGains(
+	const int32 Order, const float Azimuth, const float Elevation, TArrayView<float> OutGains)
+{
+	check(Order >= 0 && Order <= 7);
+	check(OutGains.Num() >= OrderToNumChannels(Order));
+
+	using namespace SHBasis;
+	const FTrigCache TrigCache(Azimuth, Elevation, Order);
+
+	switch (Order)
+	{
+	case 7:
+		OutGains[63] = SHBasisFunction_ACN_63(TrigCache);
+		OutGains[62] = SHBasisFunction_ACN_62(TrigCache);
+		OutGains[61] = SHBasisFunction_ACN_61(TrigCache);
+		OutGains[60] = SHBasisFunction_ACN_60(TrigCache);
+		OutGains[59] = SHBasisFunction_ACN_59(TrigCache);
+		OutGains[58] = SHBasisFunction_ACN_58(TrigCache);
+		OutGains[57] = SHBasisFunction_ACN_57(TrigCache);
+		OutGains[56] = SHBasisFunction_ACN_56(TrigCache);
+		OutGains[55] = SHBasisFunction_ACN_55(TrigCache);
+		OutGains[54] = SHBasisFunction_ACN_54(TrigCache);
+		OutGains[53] = SHBasisFunction_ACN_53(TrigCache);
+		OutGains[52] = SHBasisFunction_ACN_52(TrigCache);
+		OutGains[51] = SHBasisFunction_ACN_51(TrigCache);
+		OutGains[50] = SHBasisFunction_ACN_50(TrigCache);
+		OutGains[49] = SHBasisFunction_ACN_49(TrigCache);
+		[[fallthrough]];
+	case 6:
+		OutGains[48] = SHBasisFunction_ACN_48(TrigCache);
+		OutGains[47] = SHBasisFunction_ACN_47(TrigCache);
+		OutGains[46] = SHBasisFunction_ACN_46(TrigCache);
+		OutGains[45] = SHBasisFunction_ACN_45(TrigCache);
+		OutGains[44] = SHBasisFunction_ACN_44(TrigCache);
+		OutGains[43] = SHBasisFunction_ACN_43(TrigCache);
+		OutGains[42] = SHBasisFunction_ACN_42(TrigCache);
+		OutGains[41] = SHBasisFunction_ACN_41(TrigCache);
+		OutGains[40] = SHBasisFunction_ACN_40(TrigCache);
+		OutGains[39] = SHBasisFunction_ACN_39(TrigCache);
+		OutGains[38] = SHBasisFunction_ACN_38(TrigCache);
+		OutGains[37] = SHBasisFunction_ACN_37(TrigCache);
+		OutGains[36] = SHBasisFunction_ACN_36(TrigCache);
+		[[fallthrough]];
+	case 5:
+		OutGains[35] = SHBasisFunction_ACN_35(TrigCache);
+		OutGains[34] = SHBasisFunction_ACN_34(TrigCache);
+		OutGains[33] = SHBasisFunction_ACN_33(TrigCache);
+		OutGains[32] = SHBasisFunction_ACN_32(TrigCache);
+		OutGains[31] = SHBasisFunction_ACN_31(TrigCache);
+		OutGains[30] = SHBasisFunction_ACN_30(TrigCache);
+		OutGains[29] = SHBasisFunction_ACN_29(TrigCache);
+		OutGains[28] = SHBasisFunction_ACN_28(TrigCache);
+		OutGains[27] = SHBasisFunction_ACN_27(TrigCache);
+		OutGains[26] = SHBasisFunction_ACN_26(TrigCache);
+		OutGains[25] = SHBasisFunction_ACN_25(TrigCache);
+		[[fallthrough]];
+	case 4:
+		OutGains[24] = SHBasisFunction_ACN_24(TrigCache);
+		OutGains[23] = SHBasisFunction_ACN_23(TrigCache);
+		OutGains[22] = SHBasisFunction_ACN_22(TrigCache);
+		OutGains[21] = SHBasisFunction_ACN_21(TrigCache);
+		OutGains[20] = SHBasisFunction_ACN_20(TrigCache);
+		OutGains[19] = SHBasisFunction_ACN_19(TrigCache);
+		OutGains[18] = SHBasisFunction_ACN_18(TrigCache);
+		OutGains[17] = SHBasisFunction_ACN_17(TrigCache);
+		OutGains[16] = SHBasisFunction_ACN_16(TrigCache);
+		[[fallthrough]];
+	case 3:
+		OutGains[15] = SHBasisFunction_ACN_15(TrigCache);
+		OutGains[14] = SHBasisFunction_ACN_14(TrigCache);
+		OutGains[13] = SHBasisFunction_ACN_13(TrigCache);
+		OutGains[12] = SHBasisFunction_ACN_12(TrigCache);
+		OutGains[11] = SHBasisFunction_ACN_11(TrigCache);
+		OutGains[10] = SHBasisFunction_ACN_10(TrigCache);
+		OutGains[9]  = SHBasisFunction_ACN_9(TrigCache);
+		[[fallthrough]];
+	case 2:
+		OutGains[8]  = SHBasisFunction_ACN_8(TrigCache);
+		OutGains[7]  = SHBasisFunction_ACN_7(TrigCache);
+		OutGains[6]  = SHBasisFunction_ACN_6(TrigCache);
+		OutGains[5]  = SHBasisFunction_ACN_5(TrigCache);
+		OutGains[4]  = SHBasisFunction_ACN_4(TrigCache);
+		[[fallthrough]];
+	case 1:
+		OutGains[3]  = SHBasisFunction_ACN_3(TrigCache);
+		OutGains[2]  = SHBasisFunction_ACN_2(TrigCache);
+		OutGains[1]  = SHBasisFunction_ACN_1(TrigCache);
+		[[fallthrough]];
+	case 0:
+		OutGains[0]  = SHBasisFunction_ACN_0(TrigCache);
+		break;
+	default:
+		checkNoEntry();
+	}
+}
+
+
+void FSphericalHarmonicCalculator::GenerateFirstOrderRotationMatrixGivenRadians(const float RotXRadians, const float RotYRadians, const float RotZRadians, FMatrix & OutMatrix)
+{
+	const float SinX = FMath::Sin(RotXRadians);
+	const float SinY = FMath::Sin(RotYRadians);
+	const float SinZ = FMath::Sin(RotZRadians);
+
+	const float CosX = FMath::Cos(RotXRadians);
+	const float CosY = FMath::Cos(RotYRadians);
+	const float CosZ = FMath::Cos(RotZRadians);
+
+
+	// Build out Rotation Matrix
+	OutMatrix.SetIdentity();
+
+	// row 0:
+	/*	1.0f						0.0f															0.0f									0.0f									*/
+
+	// row 1:
+	/*	0.0f	*/	OutMatrix.M[1][1] = (CosX * CosZ) + (SinX * SinY * SinZ);		OutMatrix.M[1][2] = -SinX * CosY;		OutMatrix.M[1][3] = (CosX * SinZ) - (SinX * SinY * CosZ);
+
+	// row 2:
+	/*	0.0f	*/	OutMatrix.M[2][1] = (SinX * CosZ) - (CosX * SinY * SinZ);		OutMatrix.M[2][2] = CosX * CosY;		OutMatrix.M[2][3] = (CosX * SinY * CosZ) + (SinX * SinZ);
+
+	// row 3:
+	/*	0.0f	*/	OutMatrix.M[3][1] = -CosY * SinZ;								OutMatrix.M[3][2] = -SinY;				OutMatrix.M[3][3] = CosY * CosZ;
+}
+
+void FSphericalHarmonicCalculator::GenerateFirstOrderRotationMatrixGivenDegrees(const float RotXDegrees, const float RotYDegrees, const float RotZDegrees, FMatrix & OutMatrix)
+{
+	constexpr const float DEG_2_RAD = PI / 180.0f;
+	GenerateFirstOrderRotationMatrixGivenRadians(RotXDegrees * DEG_2_RAD, RotYDegrees * DEG_2_RAD, RotZDegrees * DEG_2_RAD, OutMatrix);
+}
diff --git a/Engine/Source/Runtime/SignalProcessing/Public/DSP/SphericalHarmonicCalculator.h b/Engine/Source/Runtime/SignalProcessing/Public/DSP/SphericalHarmonicCalculator.h
new file mode 100644
--- /dev/null
+++ b/Engine/Source/Runtime/SignalProcessing/Public/DSP/SphericalHarmonicCalculator.h
@@ -0,0 +1,48 @@
+// Copyright Epic Games, Inc. All Rights Reserved.
+
+#pragma once
+
+#include "HAL/UnrealMemory.h"
+#include "Math/Matrix.h"
+#include "Math/UnrealMathSSE.h"
+#include "Math/Vector2D.h"
+
+
+/*
+	Azimuth angle is measured CCW from front.
+	Elevation is 0 horizontal plane, + is above horizontal plane
+*/
+
+class FSphericalHarmonicCalculator
+{
+public:
+	enum UE_DEPRECATED(5.7, "Not necessary, please use a raw index for ACN values") AmbiChanNumber
+	{
+		/* 0th-Order */	ACN_0 = 0,
+		/* 1st-Order */	ACN_1, ACN_2, ACN_3,
+		/* 2nd-Order */	ACN_4, ACN_5, ACN_6, ACN_7, ACN_8,
+		/* 3rd-Order */	ACN_9, ACN_10, ACN_11, ACN_12, ACN_13, ACN_14, ACN_15,
+		/* 4th-Order */	ACN_16, ACN_17, ACN_18, ACN_19, ACN_20, ACN_21, ACN_22, ACN_23, ACN_24,
+		/* 5th-Order */	ACN_25, ACN_26, ACN_27, ACN_28, ACN_29, ACN_30, ACN_31, ACN_32, ACN_33, ACN_34, ACN_35,
+	};
+
+	static SIGNALPROCESSING_API void ComputeSoundfieldChannelGains(const int32 Order, const float Azimuth, const float Elevation, float* OutGains);
+	static SIGNALPROCESSING_API void ComputeSoundfieldChannelGains(const int32 Order, const float Azimuth, const float Elevation, TArrayView<float> OutGains);
+
+	static SIGNALPROCESSING_API void GenerateFirstOrderRotationMatrixGivenRadians(const float RotXRadians, const float RotYRadians, const float RotZRadians, FMatrix& OutMatrix);
+	static SIGNALPROCESSING_API void GenerateFirstOrderRotationMatrixGivenDegrees(const float RotXDegrees, const float RotYDegrees, const float RotZDegrees, FMatrix& OutMatrix);
+
+	template<typename T>
+	static void AdjustUESphericalCoordinatesForAmbisonics(UE::Math::TVector2<T>& InOutVector)
+	{
+		InOutVector.X = -(InOutVector.X - HALF_PI);
+		InOutVector.Y *= -1.0f;
+		Swap(InOutVector.X,InOutVector.Y);
+	}
+	
+	static int32 OrderToNumChannels(const int32 InOrder)
+	{
+		check(InOrder >= 0);
+		return (InOrder+1) * (InOrder+1);
+	}
+};
diff --git a/Engine/Source/Runtime/SoundFieldRendering/Private/SphericalHarmonicCalculator.cpp b/Engine/Source/Runtime/SoundFieldRendering/Private/SphericalHarmonicCalculator.cpp
deleted file mode 100644
diff --git a/Engine/Source/Runtime/SoundFieldRendering/Public/SoundFieldRendering.h b/Engine/Source/Runtime/SoundFieldRendering/Public/SoundFieldRendering.h
--- a/Engine/Source/Runtime/SoundFieldRendering/Public/SoundFieldRendering.h
+++ b/Engine/Source/Runtime/SoundFieldRendering/Public/SoundFieldRendering.h
@@ -4,7 +4,7 @@
 
 #include "CoreMinimal.h"
 #include "ISoundfieldFormat.h"
-#include "SphericalHarmonicCalculator.h"
+#include "DSP/SphericalHarmonicCalculator.h"
 #include "DSP/BufferVectorOperations.h"
 
 SOUNDFIELDRENDERING_API FName GetUnrealAmbisonicsFormatName();
diff --git a/Engine/Source/Runtime/SoundFieldRendering/Public/SphericalHarmonicCalculator.h b/Engine/Source/Runtime/SoundFieldRendering/Public/SphericalHarmonicCalculator.h
--- a/Engine/Source/Runtime/SoundFieldRendering/Public/SphericalHarmonicCalculator.h
+++ b/Engine/Source/Runtime/SoundFieldRendering/Public/SphericalHarmonicCalculator.h
@@ -2,41 +2,8 @@
 
 #pragma once
 
-#include "CoreMinimal.h"
-#include "DSP/BufferVectorOperations.h"
-#include "HAL/UnrealMemory.h"
-#include "Math/Matrix.h"
-#include "Math/UnrealMathSSE.h"
-#include "Math/Vector2D.h"
+#include "Misc/CoreMiscDefines.h"
 
-
-/*
-	Azimuth angle is measured CCW from front.
-	Elevation is 0 horizontal plane, + is above horizontal plane
-*/
-
-class FSphericalHarmonicCalculator
-{
-public:
-	enum AmbiChanNumber
-	{
-		/* 0th-Order */	ACN_0 = 0,
-		/* 1st-Order */	ACN_1, ACN_2, ACN_3,
-		/* 2nd-Order */	ACN_4, ACN_5, ACN_6, ACN_7, ACN_8,
-		/* 3rd-Order */	ACN_9, ACN_10, ACN_11, ACN_12, ACN_13, ACN_14, ACN_15,
-		/* 4th-Order */	ACN_16, ACN_17, ACN_18, ACN_19, ACN_20, ACN_21, ACN_22, ACN_23, ACN_24,
-		/* 5th-Order */	ACN_25, ACN_26, ACN_27, ACN_28, ACN_29, ACN_30, ACN_31, ACN_32, ACN_33, ACN_34, ACN_35
-	};
-
-	static SOUNDFIELDRENDERING_API void ComputeSoundfieldChannelGains(const int32 Order, const float Azimuth, const float Elevation, float* OutGains);
-
-	static SOUNDFIELDRENDERING_API void GenerateFirstOrderRotationMatrixGivenRadians(const float RotXRadians, const float RotYRadians, const float RotZRadians, FMatrix& OutMatrix);
-	static SOUNDFIELDRENDERING_API void GenerateFirstOrderRotationMatrixGivenDegrees(const float RotXDegrees, const float RotYDegrees, const float RotZDegrees, FMatrix& OutMatrix);
-
-	static void AdjustUESphericalCoordinatesForAmbisonics(FVector2D& InOutVector)
-	{
-		InOutVector.X = -(InOutVector.X - HALF_PI);
-		InOutVector.Y *= -1.0f;
-		Swap(InOutVector.X,InOutVector.Y);
-	}
-};
+// This file has moved.
+UE_DEPRECATED_HEADER(5.7, "This file has moved to the SignalProcessingModule")
+#include "DSP/SphericalHarmonicCalculator.h"
