Remove LFO Init function parameter defaults on that are specified in underlying IOscBase to adhere to PVS 7.36. There is a descrepancy between the default for ModMatrixStage, but codebase was audited and in practice, this was either never utilized or wasn't resulting in any change in underlying behavior.
[RN] Fixed issue with ModMatrixStage default in FLFO init function mismatch in order to adhere to codebase static analysis standards. This caused no discernible differences within the UE codebase, but is public API, so may cause unintended consequences in client code inheriting from or utilizing the ModMatrix feature in the FLFO class.

diff --git a/Engine/Source/Runtime/SignalProcessing/Public/DSP/LFO.h b/Engine/Source/Runtime/SignalProcessing/Public/DSP/LFO.h
--- a/Engine/Source/Runtime/SignalProcessing/Public/DSP/LFO.h
+++ b/Engine/Source/Runtime/SignalProcessing/Public/DSP/LFO.h
@@ -48,7 +48,7 @@ namespace Audio
 		SIGNALPROCESSING_API virtual ~FLFO();
 
 		//~ Begin FOscBase
-		SIGNALPROCESSING_API virtual void Init(const float InSampleRate, const int32 InVoiceId = 0, FModulationMatrix* InMatrix = nullptr, const int32 ModMatrixStage = 0) override;
+		SIGNALPROCESSING_API virtual void Init(const float InSampleRate, const int32 InVoiceId = 0, FModulationMatrix* InMatrix = nullptr, const int32 ModMatrixStage = 1) override;
 		SIGNALPROCESSING_API virtual void Start() override;
 		SIGNALPROCESSING_API virtual void Stop() override;
 		SIGNALPROCESSING_API virtual void Reset() override;
