Memory improvements for metasound nodes when cached:
- reset buffers that don't need to persist when cached in the wave player node
- fixed a ref counting issue in the audio bus system keeping PatchOutputs alive.

diff --git a/Engine/Plugins/Runtime/Metasound/Source/MetasoundEngine/Private/MetasoundWavePlayerNode.cpp b/Engine/Plugins/Runtime/Metasound/Source/MetasoundEngine/Private/MetasoundWavePlayerNode.cpp
--- a/Engine/Plugins/Runtime/Metasound/Source/MetasoundEngine/Private/MetasoundWavePlayerNode.cpp
+++ b/Engine/Plugins/Runtime/Metasound/Source/MetasoundEngine/Private/MetasoundWavePlayerNode.cpp
@@ -579,11 +579,15 @@ namespace Metasound
 			Resampler.Reset();
 
 			SortedCuePoints.Reset();
+			InterleavedBuffer.Empty();
+			DeinterleavedBuffer.Empty();
+			
 			for (Audio::TCircularAudioBuffer<float>& Buffer : SourceCircularBuffer)
 			{
 				Buffer.Empty();
 			}
 
+
 			SourceState = WavePlayerNodePrivate::FSourceBufferState();
 			SampleRateFrameRatio = 1.f;
 			bOnNearlyDoneTriggeredForWave = false;
diff --git a/Engine/Source/Runtime/SignalProcessing/Private/MultithreadedPatching.cpp b/Engine/Source/Runtime/SignalProcessing/Private/MultithreadedPatching.cpp
--- a/Engine/Source/Runtime/SignalProcessing/Private/MultithreadedPatching.cpp
+++ b/Engine/Source/Runtime/SignalProcessing/Private/MultithreadedPatching.cpp
@@ -263,17 +263,7 @@ namespace Audio
 
 	FPatchInput::~FPatchInput()
 	{
-		if (OutputHandle.IsValid())
-		{
-			if (OutputHandle->NumAliveInputs.fetch_sub(1, std::memory_order_relaxed) == 1)
-			{
-				if (FEvent* Event = OutputHandle->SamplesPushedEvent.load(std::memory_order_acquire))
-				{
-					// Signal that the sample status should be checked.
-					Event->Trigger();
-				}
-			}
-		}
+		Reset();
 	}
 
 	int32 FPatchInput::PushAudio(const float* InBuffer, int32 NumSamples)
@@ -346,7 +336,18 @@ namespace Audio
 
 	void FPatchInput::Reset()
 	{
-		OutputHandle.Reset();
+		if (OutputHandle.IsValid())
+		{
+			if (OutputHandle->NumAliveInputs.fetch_sub(1, std::memory_order_relaxed) == 1)
+			{
+				if (FEvent* Event = OutputHandle->SamplesPushedEvent.load(std::memory_order_acquire))
+				{
+					// Signal that the sample status should be checked.
+					Event->Trigger();
+				}
+			}
+			OutputHandle.Reset();
+		}
 	}
 
 	FPatchInput FPatchMixer::AddNewInput(int32 InMaxLatencyInSamples, float InGain)
