Fix Sine Curves on FVolumeFader when used outside [0,1] range

diff --git a/Engine/Source/Runtime/SignalProcessing/Private/VolumeFader.cpp b/Engine/Source/Runtime/SignalProcessing/Private/VolumeFader.cpp
--- a/Engine/Source/Runtime/SignalProcessing/Private/VolumeFader.cpp
+++ b/Engine/Source/Runtime/SignalProcessing/Private/VolumeFader.cpp
@@ -16,7 +16,7 @@ namespace Audio
 	{
 	}
 
-	float FVolumeFader::AlphaToVolume(float InAlpha, EFaderCurve InCurve)
+	float FVolumeFader::AlphaToVolume(float InAlpha, EFaderCurve InCurve, float InStartingValue, float InTarget)
 	{
 		switch (InCurve)
 		{
@@ -27,13 +27,17 @@ namespace Audio
 
 		case EFaderCurve::SCurve:
 		{
-			float Volume = 0.5f * Audio::FastSin(PI * InAlpha - HALF_PI) + 0.5f;
+			const float NormalizedAlpha = FMath::GetMappedRangeValueClamped(FVector2f(InStartingValue, InTarget), FVector2f(0.0f, 1.0f), InAlpha);
+			float Volume = 0.5f * Audio::FastSin(PI * NormalizedAlpha - HALF_PI) + 0.5f;
+			Volume = FMath::GetMappedRangeValueClamped(FVector2f(0.0f, 1.0f), FVector2f(InStartingValue, InTarget), Volume);
 			return FMath::Max(0.0f, Volume);
 		}
 
 		case EFaderCurve::Sin:
 		{
-			float Volume = Audio::FastSin(HALF_PI * InAlpha);
+			const float NormalizedAlpha = FMath::GetMappedRangeValueClamped(FVector2f(InStartingValue, InTarget), FVector2f(0.0f, 1.0f), InAlpha);
+			float Volume = Audio::FastSin(HALF_PI * NormalizedAlpha);
+			Volume = FMath::GetMappedRangeValueClamped(FVector2f(0.0f, 1.0f), FVector2f(InStartingValue, InTarget), Volume);
 			return FMath::Max(0.0f, Volume);
 		}
 
@@ -73,7 +77,7 @@ namespace Audio
 
 	float FVolumeFader::GetVolume() const
 	{
-		return AlphaToVolume(Alpha, FadeCurve);
+		return AlphaToVolume(Alpha, FadeCurve, StartingValue, Target);
 	}
 
 	float FVolumeFader::GetVolumeAfterTime(float InDeltaTime) const
@@ -110,7 +114,7 @@ namespace Audio
 			FutureAlpha = FMath::Clamp(FutureAlpha, MinValue, MaxValue);
 		}
 
-		return AlphaToVolume(FutureAlpha, FadeCurve);
+		return AlphaToVolume(FutureAlpha, FadeCurve, StartingValue, Target);
 	}
 
 	float FVolumeFader::GetFadeDuration() const
@@ -220,6 +224,7 @@ namespace Audio
 		Elapsed = 0.0f;
 		FadeCurve = InCurve;
 		FadeDuration = InDuration;
+		StartingValue = Alpha;
 	}
 
 	void FVolumeFader::StopFade()
diff --git a/Engine/Source/Runtime/SignalProcessing/Public/DSP/VolumeFader.h b/Engine/Source/Runtime/SignalProcessing/Public/DSP/VolumeFader.h
--- a/Engine/Source/Runtime/SignalProcessing/Public/DSP/VolumeFader.h
+++ b/Engine/Source/Runtime/SignalProcessing/Public/DSP/VolumeFader.h
@@ -128,7 +128,7 @@ namespace Audio
 
 	private:
 		/** Converts value to final resulting volume */
-		static float AlphaToVolume(float InAlpha, EFaderCurve InCurve);
+		static float AlphaToVolume(float InAlpha, EFaderCurve InCurve, float InStartingValue, float InTarget);
 
 		/** Current value used to linear interpolate over update delta
 		  * (Normalized value for non-log, -80dB to 0dB for log)
@@ -151,5 +151,8 @@ namespace Audio
 
 		/** Audio fader curve to use */
 		EFaderCurve FadeCurve;
+
+		/** Initial value from the start of the fade */
+		float StartingValue;
 	};
 } // namespace Audio
