Fix for FDelay::SetEasedDelayMsec() causing read index to overtake write index.

diff --git a/Engine/Source/Runtime/SignalProcessing/Private/Delay.cpp b/Engine/Source/Runtime/SignalProcessing/Private/Delay.cpp
--- a/Engine/Source/Runtime/SignalProcessing/Private/Delay.cpp
+++ b/Engine/Source/Runtime/SignalProcessing/Private/Delay.cpp
@@ -184,6 +184,13 @@ namespace Audio
 	{
 		const float DesiredDelayInSamples = InDelayMsec * SampleRate * 0.001f;
 		const float TargetDelayInSamples = FMath::Min(DesiredDelayInSamples, MaxBufferLengthSamples - 1);
+
+		if (FMath::IsNearlyEqual(TargetDelayInSamples, DesiredDelayInSamples))
+		{
+			SetDelaySamples(TargetDelayInSamples);
+			return;
+		}
+		
 		ResizeIfNeeded(TargetDelayInSamples);
 
 		EaseDelayMsec.SetValue(InDelayMsec, bIsInit);
