Add support to "Convert Deinterleave" to process MultiChannel audio as array of ArrayViews

#rnx
#tests Audio Unit Tests

diff --git a/Engine/Source/Runtime/SignalProcessing/Private/ConvertDeinterleave.cpp b/Engine/Source/Runtime/SignalProcessing/Private/ConvertDeinterleave.cpp
--- a/Engine/Source/Runtime/SignalProcessing/Private/ConvertDeinterleave.cpp
+++ b/Engine/Source/Runtime/SignalProcessing/Private/ConvertDeinterleave.cpp
@@ -52,19 +52,19 @@ namespace Audio
 				return &ChannelGains.GetData()[InOutputChannelIndex * NumInputChannels];
 			}
 
-			virtual void ProcessAudio(TArrayView<const float> InSamples, FMultichannelBuffer& OutSamples) const override
+			virtual void ProcessAudio(TArrayView<const float> InSamples, TArray<TArrayView<float>>& OutMultiChannelAudio) const override
 			{
 				checkf((InSamples.Num() % NumInputChannels) == 0, TEXT("Input sample buffer contains partial audio frame."));
-				checkf((InSamples.Num() / NumInputChannels) == GetMultichannelBufferNumFrames(OutSamples), TEXT("Audio buffer frame count mismatch."));
-				checkf(NumOutputChannels == OutSamples.Num(), TEXT("Output audio buffer not initialized to expected channel count."));
+				checkf((InSamples.Num() / NumInputChannels) == GetMultichannelBufferNumFrames(OutMultiChannelAudio), TEXT("Audio buffer frame count mismatch."));
+				checkf(NumOutputChannels == OutMultiChannelAudio.Num(), TEXT("Output audio buffer not initialized to expected channel count."));
 
-				const int32 NumFrames = GetMultichannelBufferNumFrames(OutSamples);
+				const int32 NumFrames = GetMultichannelBufferNumFrames(OutMultiChannelAudio);
 				const float* InSampleData = InSamples.GetData();
 
 				float* OutputChannelDataPtrs[NumOutputChannels];
 				for (int32 OutChannelIndex = 0; OutChannelIndex < NumOutputChannels; OutChannelIndex++)
 				{
-					OutputChannelDataPtrs[OutChannelIndex] = OutSamples[OutChannelIndex].GetData();
+					OutputChannelDataPtrs[OutChannelIndex] = OutMultiChannelAudio[OutChannelIndex].GetData();
 				}
 
 				for (int32 Frame = 0; Frame < NumFrames; Frame++)
@@ -87,6 +87,7 @@ namespace Audio
 					}
 				}
 			}
+			
 		private:
 			TArray<float> ChannelGains;
 		};
@@ -119,27 +120,30 @@ namespace Audio
 
 			virtual ~TConvertDeinterleave() = default;
 
-			virtual void ProcessAudio(TArrayView<const float> InSamples, FMultichannelBuffer& OutSamples) const override
+			virtual void ProcessAudio(TArrayView<const float> InSamples, TArray<TArrayView<float>>& OutMultiChannelAudio) const override
 			{
-				checkf(InSamples.Num() == GetMultichannelBufferNumFrames(OutSamples), TEXT("Audio buffer frame count mismatch."));
-				checkf(NumOutputChannels == OutSamples.Num(), TEXT("Output audio buffer not initialized to expected channel count."));
+				checkf(InSamples.Num() == GetMultichannelBufferNumFrames(OutMultiChannelAudio), TEXT("Audio buffer frame count mismatch."));
+				checkf(NumOutputChannels == OutMultiChannelAudio.Num(), TEXT("Output audio buffer not initialized to expected channel count."));
 
-				const int32 NumFrames = GetMultichannelBufferNumFrames(OutSamples);
+				const int32 NumFrames = GetMultichannelBufferNumFrames(OutMultiChannelAudio);
 
 				for (int32 OutChannelIndex = 0; OutChannelIndex < NumOutputChannels; OutChannelIndex++)
 				{
 					const float Gain = ChannelGains[OutChannelIndex];
-					if (Gain != 0.f)
+					if (Gain == 1.f)
+					{
+						// Simply copy the audio to the output channels when gain is 1
+						FMemory::Memcpy(OutMultiChannelAudio[OutChannelIndex].GetData(), InSamples.GetData(), NumFrames * sizeof(float));
+					}
+					else if (Gain != 0.f)
 					{
-						TArrayView<float> OutSampleView(OutSamples[OutChannelIndex].GetData(), NumFrames);
-
 						// Only need to multiply if gain is non-zero
-						Audio::ArrayMultiplyByConstant(InSamples, ChannelGains[OutChannelIndex], OutSampleView);
+						Audio::ArrayMultiplyByConstant(InSamples, ChannelGains[OutChannelIndex], OutMultiChannelAudio[OutChannelIndex]);
 					}
 					else
 					{
 						// If gain is zero, we can set the output channel to zero.
-						FMemory::Memset(OutSamples[OutChannelIndex].GetData(), 0, NumFrames * sizeof(float));
+						FMemory::Memset(OutMultiChannelAudio[OutChannelIndex].GetData(), 0, NumFrames * sizeof(float));
 					}
 				}
 			}
@@ -160,14 +164,14 @@ namespace Audio
 			}
 
 			virtual ~TConvertDeinterleave() = default;
-			virtual void ProcessAudio(TArrayView<const float> InAudio, FMultichannelBuffer& OutAudio) const override
+			virtual void ProcessAudio(TArrayView<const float> InSamples, TArray<TArrayView<float>>& OutMultiChannelAudio) const override
 			{
-				checkf(InAudio.Num() == GetMultichannelBufferNumFrames(OutAudio), TEXT("Audio buffer frame count mismatch."));
-				checkf(1 == OutAudio.Num(), TEXT("Output audio buffer not initialized to expected channel count."));
+				checkf(InSamples.Num() == GetMultichannelBufferNumFrames(OutMultiChannelAudio), TEXT("Audio buffer frame count mismatch."));
+				checkf(1 == OutMultiChannelAudio.Num(), TEXT("Output audio buffer not initialized to expected channel count."));
 
-				const int32 NumFrames = GetMultichannelBufferNumFrames(OutAudio);
-				const float* InSampleData = InAudio.GetData();
-				FMemory::Memcpy(OutAudio[0].GetData(), InSampleData, NumFrames * sizeof(float));
+				const int32 NumFrames = GetMultichannelBufferNumFrames(OutMultiChannelAudio);
+				const float* InSampleData = InSamples.GetData();
+				FMemory::Memcpy(OutMultiChannelAudio[0].GetData(), InSampleData, NumFrames * sizeof(float));
 			}
 		};
 
diff --git a/Engine/Source/Runtime/SignalProcessing/Public/DSP/ConvertDeinterleave.h b/Engine/Source/Runtime/SignalProcessing/Public/DSP/ConvertDeinterleave.h
--- a/Engine/Source/Runtime/SignalProcessing/Public/DSP/ConvertDeinterleave.h
+++ b/Engine/Source/Runtime/SignalProcessing/Public/DSP/ConvertDeinterleave.h
@@ -43,7 +43,22 @@ namespace Audio
 		 *                     the input samples and possibly upmixing or downmixing
 		 *                     them to the target channel count. 
 		 */
-		virtual void ProcessAudio(TArrayView<const float> InSamples, FMultichannelBuffer& OutSamples) const = 0;
+		void ProcessAudio(TArrayView<const float> InSamples, FMultichannelBuffer& OutSamples) const
+		{
+			TArray<TArrayView<float>> OutMultiChannelAudio = MakeMultichannelBufferView(OutSamples);
+			ProcessAudio(InSamples,OutMultiChannelAudio);
+		}
+
+		/** Deinterleave and convert the channel format of the input audio.
+		 *
+		 * @param InSamples - ArrayView of interleaved input samples. The number of 
+		 *                    samples must be evenly divisible by the number of 
+		 *                    input channels. 
+		 * @param OutSamples - A multichannel buffer of samples generated by deinterleaving
+		 *                     the input samples and possibly upmixing or downmixing
+		 *                     them to the target channel count. 
+		 */
+		virtual void ProcessAudio(TArrayView<const float> InSamples, TArray<TArrayView<float>>& OutMultiChannelAudio) const = 0;
 
 
 		/** Create an IConvertDeinterleave object for a given number of input
