ConvertDeinterleave

update code comment for function to reflect correct input name

#rnx

diff --git a/Engine/Source/Runtime/SignalProcessing/Public/DSP/ConvertDeinterleave.h b/Engine/Source/Runtime/SignalProcessing/Public/DSP/ConvertDeinterleave.h
--- a/Engine/Source/Runtime/SignalProcessing/Public/DSP/ConvertDeinterleave.h
+++ b/Engine/Source/Runtime/SignalProcessing/Public/DSP/ConvertDeinterleave.h
@@ -54,7 +54,7 @@ namespace Audio
 		 * @param InSamples - ArrayView of interleaved input samples. The number of 
 		 *                    samples must be evenly divisible by the number of 
 		 *                    input channels. 
-		 * @param OutSamples - A multichannel buffer of samples generated by deinterleaving
+		 * @param OutMultiChannelAudio - A multichannel buffer of samples generated by deinterleaving
 		 *                     the input samples and possibly upmixing or downmixing
 		 *                     them to the target channel count. 
 		 */
