[SignalProcessing] Added EPeakMode::TruePeak to envelope follower. (Submit on behalf of pete.bliss) + added true peak extreme test case

#Synced-CL 49945202

diff --git a/Engine/Source/Runtime/SignalProcessing/Private/EnvelopeFollower.cpp b/Engine/Source/Runtime/SignalProcessing/Private/EnvelopeFollower.cpp
--- a/Engine/Source/Runtime/SignalProcessing/Private/EnvelopeFollower.cpp
+++ b/Engine/Source/Runtime/SignalProcessing/Private/EnvelopeFollower.cpp
@@ -23,6 +23,38 @@ namespace Audio
 
 			return Value;
 		}
+
+		/**
+		 * Returns the maximum of the elements in the given vector
+		 */
+		FORCEINLINE float VectorReduceMax(const VectorRegister4Float& ABCD)
+		{
+			const VectorRegister4Float BADC = VectorSwizzle(ABCD, 1, 0, 3, 2);
+			const VectorRegister4Float MaxAB_MaxCD = VectorMax(ABCD, BADC);
+			const VectorRegister4Float MaxCD_MaxAB = VectorSwizzle(MaxAB_MaxCD, 2, 3, 0, 1);
+			const VectorRegister4Float MaxABCD = VectorMax(MaxAB_MaxCD, MaxCD_MaxAB);
+			return VectorGetComponent(MaxABCD, 0);
+		}
+
+		alignas(AUDIO_SIMD_BYTE_ALIGNMENT) static const float TruePeakPolyphaseCoefficients[] =
+		{
+			+0.0017089843750f, -0.0291748046875f, -0.0189208984375f, -0.0083007812500f,
+			+0.0109863281250f, +0.0292968750000f, +0.0330810546875f, +0.0148925781250f,
+			-0.0196533203125f, -0.0517578125000f, -0.0582275390625f, -0.0266113281250f,
+			+0.0332031250000f, +0.0891113281250f, +0.1015625000000f, +0.0476074218750f,
+			-0.0594482421875f, -0.1665039062500f, -0.2003173828125f, -0.1022949218750f,
+			+0.1373291015625f, +0.4650878906250f, +0.7797851562500f, +0.9721679687500f,
+			+0.9721679687500f, +0.7797851562500f, +0.4650878906250f, +0.1373291015625f,
+			-0.1022949218750f, -0.2003173828125f, -0.1665039062500f, -0.0594482421875f,
+			+0.0476074218750f, +0.1015625000000f, +0.0891113281250f, +0.0332031250000f,
+			-0.0266113281250f, -0.0582275390625f, -0.0517578125000f, -0.0196533203125f,
+			+0.0148925781250f, +0.0330810546875f, +0.0292968750000f, +0.0109863281250f,
+			-0.0083007812500f, -0.0189208984375f, -0.0291748046875f, +0.0017089843750f,
+		};
+
+		constexpr uint32 TruePeakPhaseCount = 4; // == Over-sampling ratio
+		constexpr uint32 TruePeakNumCoefficientsPerPhase = UE_ARRAY_COUNT(TruePeakPolyphaseCoefficients) / TruePeakPhaseCount;
+		constexpr uint32 TruePeakFilterHistorySize = 1 << FMath::CeilLogTwo(TruePeakNumCoefficientsPerPhase - 1);
 	}
 
 	FAttackRelease::FAttackRelease(float InSampleRate, float InAttackTimeMsec, float InReleaseTimeMsec, bool bInIsAnalog)
@@ -335,6 +367,60 @@ namespace Audio
 		}
 	}
 
+	void FTruePeak::Reset()
+	{
+		AllChannelSignalHistories.Reset();
+		AllChannelSignalHistories.SetNumZeroed(NumChannels * EnvelopeFollowerPrivate::TruePeakFilterHistorySize);
+		HistoryPosition = 0;
+	}
+
+	void FTruePeak::ProcessAudio(const float* InBuffer, int32 InNumFrames, float* OutBuffer)
+	{
+		using namespace EnvelopeFollowerPrivate;
+
+		constexpr uint32 HistoryIndexMask = TruePeakFilterHistorySize - 1;
+
+		const int32 NumSamples = InNumFrames * NumChannels;
+
+		uint32 PrevHistoryWriteIndex = HistoryPosition;
+
+		for (int32 ChannelIndex = 0; ChannelIndex < NumChannels; ChannelIndex++)
+		{
+			const TArrayView ChannelSignalHistory = MakeArrayView(AllChannelSignalHistories).Slice(ChannelIndex * TruePeakFilterHistorySize, TruePeakFilterHistorySize);
+			
+			// Reset PrevHistoryWriteIndex for each channel:
+			PrevHistoryWriteIndex = HistoryPosition;
+
+			for (int32 SampleIndex = ChannelIndex; SampleIndex < NumSamples; SampleIndex += NumChannels)
+			{
+				const float InputSample = InBuffer[SampleIndex];
+
+				static_assert(TruePeakPhaseCount == sizeof(VectorRegister4Float) / sizeof(float)); // TruePeakPhaseCount matches SIMD size.
+				const VectorRegister4Float PolyphaseCoefficients0 = VectorLoad(&TruePeakPolyphaseCoefficients[0]);
+
+				VectorRegister4Float Accumulators = VectorMultiply(VectorSetFloat1(InputSample), PolyphaseCoefficients0);
+
+				uint32 HistoryReadIndex = PrevHistoryWriteIndex;
+				for (uint32 CoefficientReadIndex = TruePeakPhaseCount; CoefficientReadIndex < UE_ARRAY_COUNT(TruePeakPolyphaseCoefficients); CoefficientReadIndex += TruePeakPhaseCount)
+				{
+					const VectorRegister4Float PolyphaseCoefficients = VectorLoad(&TruePeakPolyphaseCoefficients[CoefficientReadIndex]);
+					const VectorRegister4Float HistorySample = VectorLoadFloat1(&ChannelSignalHistory[HistoryReadIndex]);
+					Accumulators = VectorMultiplyAdd(HistorySample, PolyphaseCoefficients, Accumulators);
+					HistoryReadIndex = (HistoryReadIndex - 1) & HistoryIndexMask; // uint32 wraparound is fine here as a bitmask is applied.
+				}
+
+				const float MaxAbsValue = VectorReduceMax(VectorAbs(Accumulators));
+				OutBuffer[SampleIndex] = MaxAbsValue;
+
+				const uint32 HistoryWriteIndex = (PrevHistoryWriteIndex + 1) & HistoryIndexMask;
+				ChannelSignalHistory[HistoryWriteIndex] = InputSample;
+				PrevHistoryWriteIndex = HistoryWriteIndex;
+			}
+		}
+
+		HistoryPosition = PrevHistoryWriteIndex;
+	}
+
 
 	// A simple utility that returns a smoothed value given audio input using an RC circuit.
 	// Used for following the envelope of an audio stream.
@@ -345,6 +431,7 @@ namespace Audio
 
 	FEnvelopeFollower::FEnvelopeFollower(const FEnvelopeFollowerInitParams& InParams)
 	: MeanSquaredProcessor(InParams.SampleRate, InParams.NumChannels, InParams.AnalysisWindowMsec)
+	, TruePeakProcessor((InParams.Mode == EPeakMode::TruePeak) ? InParams.NumChannels : 0)
 	, Smoother(InParams.SampleRate, InParams.NumChannels, InParams.AttackTimeMsec, InParams.ReleaseTimeMsec, InParams.bIsAnalog)
 	, NumChannels(InParams.NumChannels)
 	, EnvMode(InParams.Mode)
@@ -356,6 +443,7 @@ namespace Audio
 	{
 		Smoother = FAttackReleaseSmoother(InParams.SampleRate, InParams.NumChannels, InParams.AttackTimeMsec, InParams.ReleaseTimeMsec, InParams.bIsAnalog);
 		MeanSquaredProcessor = FMeanSquaredIIR(InParams.SampleRate, InParams.NumChannels, InParams.AnalysisWindowMsec);
+		TruePeakProcessor = FTruePeak((InParams.Mode == EPeakMode::TruePeak) ? InParams.NumChannels : 0);
 		NumChannels = InParams.NumChannels;
 		EnvMode = InParams.Mode;
 	}
@@ -395,7 +483,8 @@ namespace Audio
 		if (InNumChannels != NumChannels)
 		{
 			Smoother.SetNumChannels(InNumChannels);
-			MeanSquaredProcessor.SetNumChannels(NumChannels);
+			MeanSquaredProcessor.SetNumChannels(InNumChannels);
+			TruePeakProcessor.SetNumChannels((EnvMode == EPeakMode::TruePeak) ? InNumChannels : 0);
 			NumChannels = InNumChannels;
 		}
 	}
@@ -404,6 +493,7 @@ namespace Audio
 	void FEnvelopeFollower::Reset()
 	{
 		MeanSquaredProcessor.Reset();
+		TruePeakProcessor.Reset();
 		Smoother.Reset();
 	}
 
@@ -429,6 +519,8 @@ namespace Audio
 	void FEnvelopeFollower::SetMode(EPeakMode::Type InMode)
 	{
 		EnvMode = InMode;
+
+		TruePeakProcessor.SetNumChannels((InMode == EPeakMode::TruePeak) ? NumChannels : 0);
 	}
 
 	// Process the input audio buffer and returns the last envelope value
@@ -473,6 +565,10 @@ namespace Audio
 			{
 				Audio::ArrayAbs(MakeArrayView(InBuffer, NumSamples), WorkBuffer);
 			}
+			else if (EPeakMode::TruePeak == EnvMode)
+			{
+				TruePeakProcessor.ProcessAudio(InBuffer, InNumFrames, WorkBuffer.GetData());
+			}
 			else
 			{
 				check(false);
diff --git a/Engine/Source/Runtime/SignalProcessing/Public/DSP/EnvelopeFollower.h b/Engine/Source/Runtime/SignalProcessing/Public/DSP/EnvelopeFollower.h
--- a/Engine/Source/Runtime/SignalProcessing/Public/DSP/EnvelopeFollower.h
+++ b/Engine/Source/Runtime/SignalProcessing/Public/DSP/EnvelopeFollower.h
@@ -18,6 +18,7 @@ namespace Audio
 			MeanSquared,
 			RootMeanSquared,
 			Peak,
+			TruePeak,
 			Count
 		};
 	}
@@ -201,6 +202,37 @@ namespace Audio
 		TArray<float> ChannelValues;
 	};
 
+	/** Compute True-Peak as defined by ITU-R BS.1770-5 */
+	class FTruePeak
+	{
+	public:
+		FTruePeak(int32 InNumChannels)
+			: NumChannels(InNumChannels)
+		{
+			if (InNumChannels > 0)
+			{
+				Reset();
+			}
+		}
+
+		void SetNumChannels(int32 InNumChannels)
+		{
+			if (NumChannels != InNumChannels)
+			{
+				NumChannels = InNumChannels;
+				Reset();
+			}
+		}
+
+		SIGNALPROCESSING_API void Reset();
+		SIGNALPROCESSING_API void ProcessAudio(const float* InBuffer, int32 InNumFrames, float* OutBuffer);
+
+	private:
+		int32 NumChannels = 0;
+		TArray<float> AllChannelSignalHistories;
+		uint32 HistoryPosition = 0;
+	};
+
 	struct FEnvelopeFollowerInitParams
 	{
 		/** Number of frames per a second. */
@@ -293,6 +325,7 @@ namespace Audio
 
 		FAlignedFloatBuffer WorkBuffer;
 		FMeanSquaredIIR MeanSquaredProcessor; 
+		FTruePeak TruePeakProcessor;
 		FAttackReleaseSmoother Smoother;
 
 		int32 NumChannels;
@@ -328,10 +361,11 @@ namespace Audio
 		FInlineEnvelopeFollower(const FInlineEnvelopeFollowerInitParams& InParams)
 		: FAttackRelease(InParams.SampleRate, InParams.AttackTimeMsec, InParams.ReleaseTimeMsec, InParams.bIsAnalog)
 		, Value(0.f)
-		, Mode(InParams.Mode)
+		, Mode((InParams.Mode != EPeakMode::TruePeak) ? InParams.Mode : EPeakMode::Peak)
 		, AnalysisValue(0.f)
 		, AnalysisWindowMsec(5.f)
 		{
+			ensure(InParams.Mode != EPeakMode::TruePeak);
 			SetAnalysisWindow(InParams.AnalysisWindowMsec);
 		}
 
@@ -343,10 +377,11 @@ namespace Audio
 		/** Initialize an envelope follower.  */
 		void Init(const FInlineEnvelopeFollowerInitParams& InParams)
 		{
+			ensure(InParams.Mode != EPeakMode::TruePeak);
 			SetSampleRate(InParams.SampleRate);
 			SetAttackTime(InParams.AttackTimeMsec);
 			SetReleaseTime(InParams.ReleaseTimeMsec);
-			SetMode(InParams.Mode);
+			SetMode((InParams.Mode != EPeakMode::TruePeak) ? InParams.Mode : EPeakMode::Peak);
 			SetAnalog(InParams.bIsAnalog);
 			SetAnalysisWindow(InParams.AnalysisWindowMsec);
 		}
@@ -354,7 +389,10 @@ namespace Audio
 		/* Sets the input analysis mode of the envelope follower */
 		void SetMode(EPeakMode::Type InMode)
 		{
-			Mode = InMode;
+			if (ensure(InMode != EPeakMode::TruePeak))
+			{
+				Mode = InMode;
+			}
 		}
 
 		/** Set the analysis window size (for MeanSquared and RootMeanSquared). */
@@ -442,6 +480,7 @@ namespace Audio
 						return FMath::Sqrt(AnalysisValue);
 					}
 
+				case EPeakMode::TruePeak:
 				default:
 				{
 					checkNoEntry();
@@ -485,6 +524,7 @@ namespace Audio
 				return;
 			}
 
+			case EPeakMode::TruePeak:
 			default:
 			{
 				checkNoEntry();
