Add unit test to catch an edge case in runtime resampler, and fix the edge case.

#rnx

#Synced-CL 50251988
#tests

#Synced-CL 50251988

diff --git a/Engine/Source/Runtime/SignalProcessing/Private/RuntimeResampler.cpp b/Engine/Source/Runtime/SignalProcessing/Private/RuntimeResampler.cpp
--- a/Engine/Source/Runtime/SignalProcessing/Private/RuntimeResampler.cpp
+++ b/Engine/Source/Runtime/SignalProcessing/Private/RuntimeResampler.cpp
@@ -126,9 +126,10 @@ namespace Audio
 		{
 			const int64 InputFrameFP = MapOutputFrameToInputFrameFP(InNumOutputFrames - 1);
 
-			// We need the sample after InputFrameFP, so we take the floor and add 1. Because we're returning
-			//  the number of samples needed rather than the index of the last sample, we add another 1.
-			return int32(InputFrameFP / FPScale) + 2;
+			// We need the frame after InputFrameFP, so we take the floor and add 1. Because we're returning
+			//  the number of frames needed rather than the index of the last frame, we add another 1. We add these
+			//  frames before dividing to make sure the numerator is positive and the division gives us the floor.
+			return int32((InputFrameFP + 2 * FPScale) / FPScale);
 		}
 		return 0;
 	}
