Fix for transient behavior while adding/removing controllers on certain platforms causing key and output channel count to briefly differ and assert.

[RN] Fix for transient behavior while adding/removing controllers on certain platforms causing key and output channel count to briefly differ and assert in the Dynamics Processor Submix Effect.

#Synced-CL 50420887

diff --git a/Engine/Source/Runtime/SignalProcessing/Private/DynamicsProcessor.cpp b/Engine/Source/Runtime/SignalProcessing/Private/DynamicsProcessor.cpp
--- a/Engine/Source/Runtime/SignalProcessing/Private/DynamicsProcessor.cpp
+++ b/Engine/Source/Runtime/SignalProcessing/Private/DynamicsProcessor.cpp
@@ -375,8 +375,15 @@ namespace Audio
 
 	void FDynamicsProcessor::ProcessAudio(const float* const* const InBuffers, const int32 InNumFrames, float* const* OutBuffers, const float* const* const InKeyBuffers, float* const* OutEnvelopes)
 	{
-		// we need this test because we are going to use the output buffer as our scratch pad for key processing!
-		check(GetKeyNumChannels() <= GetNumChannels()); 
+		// There are rare cases where the number of output channels is mutated but the key count is
+		// not (such as during device swap, adding & removing endpoints such as controllers, etc.).
+		// For now, ensure key number of channels matches that of the output number of channels.
+		const int32 NumChannels = GetNumChannels();
+		if (GetKeyNumChannels() > NumChannels)
+		{
+			SetKeyNumChannels(NumChannels);
+		}
+
 		ProcessAudioDeinterleaveInternal(InBuffers, InNumFrames, OutBuffers, InKeyBuffers, OutEnvelopes);
 	}
 
