remove volatile from spectrum analyzer

#Synced-CL 50798981

diff --git a/Engine/Source/Runtime/SignalProcessing/Private/SpectrumAnalyzer.cpp b/Engine/Source/Runtime/SignalProcessing/Private/SpectrumAnalyzer.cpp
--- a/Engine/Source/Runtime/SignalProcessing/Private/SpectrumAnalyzer.cpp
+++ b/Engine/Source/Runtime/SignalProcessing/Private/SpectrumAnalyzer.cpp
@@ -1185,23 +1185,27 @@ namespace Audio
 
 	FAlignedFloatBuffer& FSpectrumAnalyzerBuffer::StartWorkOnBuffer()
 	{
+		FScopeLock ScopeLock(&BufferIndicesCriticalSection);
 		return ComplexBuffers[InputIndex];
 	}
 
 	void FSpectrumAnalyzerBuffer::StopWorkOnBuffer(double InTimestamp)
 	{
+		FScopeLock ScopeLock(&BufferIndicesCriticalSection);
 		Timestamps[InputIndex] = InTimestamp;
 		IncrementInputIndex();
 	}
 
 	const FAlignedFloatBuffer& FSpectrumAnalyzerBuffer::LockMostRecentBuffer(double& OutTimestamp) const
 	{
+		FScopeLock ScopeLock(&BufferIndicesCriticalSection);
 		OutTimestamp = Timestamps[OutputIndex];
 		return ComplexBuffers[OutputIndex];
 	}
 
 	const FAlignedFloatBuffer& FSpectrumAnalyzerBuffer::LockMostRecentBuffer() const
 	{
+		FScopeLock ScopeLock(&BufferIndicesCriticalSection);
 		return ComplexBuffers[OutputIndex];
 	}
 
diff --git a/Engine/Source/Runtime/SignalProcessing/Public/DSP/SpectrumAnalyzer.h b/Engine/Source/Runtime/SignalProcessing/Public/DSP/SpectrumAnalyzer.h
--- a/Engine/Source/Runtime/SignalProcessing/Public/DSP/SpectrumAnalyzer.h
+++ b/Engine/Source/Runtime/SignalProcessing/Public/DSP/SpectrumAnalyzer.h
@@ -238,11 +238,11 @@ namespace Audio
 		void IncrementInputIndex();
 		void IncrementOutputIndex();
 
-		volatile int32 OutputIndex;
-		volatile int32 InputIndex;
+		int32 OutputIndex;
+		int32 InputIndex;
 
 		// This mutex is locked when we increment either the input or output index.
-		FCriticalSection BufferIndicesCriticalSection;
+		mutable FCriticalSection BufferIndicesCriticalSection;
 	};
 
 	class FSpectrumAnalysisAsyncWorker 
@@ -285,7 +285,7 @@ namespace Audio
 	 * Class built to be a rolling spectrum analyzer for arbitrary, monaural audio data.
 	 * Class is meant to scale accuracy with CPU and memory budgets.
 	 * Typical usage is to either call PushAudio() and then PerformAnalysisIfPossible immediately afterwards,
-	 * or have a seperate thread call PerformAnalysisIfPossible().
+	 * or have a separate thread call PerformAnalysisIfPossible().
 	 */
 	class FSpectrumAnalyzer 
 	{
