Handle the single speaker case for VBap

#tests LowLevelTest (AudioUnitTests)
#Synced-CL 51449712

diff --git a/Engine/Source/Runtime/SignalProcessing/Private/Vbap.cpp b/Engine/Source/Runtime/SignalProcessing/Private/Vbap.cpp
--- a/Engine/Source/Runtime/SignalProcessing/Private/Vbap.cpp
+++ b/Engine/Source/Runtime/SignalProcessing/Private/Vbap.cpp
@@ -222,7 +222,10 @@ namespace Audio
 
 			// Check we don't have any spans bigger than 180.
 			const float SpanWidth = CalculateWidth(Speaker.Speaker.AzimuthDegrees, NextSpeaker.Speaker.AzimuthDegrees);
-			if (FMath::Abs(SpanWidth) < 180.f)
+			
+			// Only add spans that are less than 180. 
+			// In the single speaker special case, that'll be > 180, but allow it anyway.
+			if (FMath::Abs(SpanWidth) < 180.f || SortedSpeakers.Num() == 1)
 			{
 				Spans.Add(
 					{
