Add function to convert frequency to midi with non-standard A440 tuning

#Synced-CL 51555917

diff --git a/Engine/Source/Runtime/SignalProcessing/Public/DSP/Dsp.h b/Engine/Source/Runtime/SignalProcessing/Public/DSP/Dsp.h
--- a/Engine/Source/Runtime/SignalProcessing/Public/DSP/Dsp.h
+++ b/Engine/Source/Runtime/SignalProcessing/Public/DSP/Dsp.h
@@ -386,6 +386,16 @@ namespace Audio
 		return 69.0f + 12.0f * FMath::LogX(2.0f, InFrequency / 440.0f);
 	}
 
+	// Using passed in tuning frequency (standard is 440), compute midi from frequency in hz
+	inline float GetMidiFromFrequency(const float InFrequency, const float TuningFrequency)
+	{
+		if (TuningFrequency <= 0.f)
+		{
+			return 0.f;
+		}
+		return 69.0f + 12.0f * FMath::LogX(2.0f, InFrequency / TuningFrequency);
+	}
+
 	// Return a pitch scale factor based on the difference between a base midi note and a target midi note. Useful for samplers.
 	inline float GetPitchScaleFromMIDINote(int32 BaseMidiNote, int32 TargetMidiNote)
 	{
