[SignalProcessing] LKFSAnalyzer: fix in CalculateGatedAverageFromHistory so the minimum loudness is -Inf instead of gated loudness.

#Synced-CL 51662785

diff --git a/Engine/Source/Runtime/SignalProcessing/Private/LKFSAnalyzer.cpp b/Engine/Source/Runtime/SignalProcessing/Private/LKFSAnalyzer.cpp
--- a/Engine/Source/Runtime/SignalProcessing/Private/LKFSAnalyzer.cpp
+++ b/Engine/Source/Runtime/SignalProcessing/Private/LKFSAnalyzer.cpp
@@ -49,7 +49,7 @@ namespace Audio
 				}
 			}
 
-			return (Num > 0) ? Accum /= Num : InGate;
+			return (Num > 0) ? Accum /= Num : 0.0f;
 		}
 
 		float CalculateGatedLoudnessFromHistory(const TCircularBuffer<float>& InHistory, int32 InStartIndex, int32 InEndIndex, int32 InStride)
