[SignalProcessing] [AudioSynesthesia] LKFS analyzer: improving LUFS accuracy in the low frequency region to approximate as much as possible to the IUT-R BS.1770 standard.

#Synced-CL 51662785

diff --git a/Engine/Plugins/Runtime/AudioSynesthesia/Source/AudioSynesthesiaCore/Public/LKFSNRTFactory.h b/Engine/Plugins/Runtime/AudioSynesthesia/Source/AudioSynesthesiaCore/Public/LKFSNRTFactory.h
--- a/Engine/Plugins/Runtime/AudioSynesthesia/Source/AudioSynesthesiaCore/Public/LKFSNRTFactory.h
+++ b/Engine/Plugins/Runtime/AudioSynesthesia/Source/AudioSynesthesiaCore/Public/LKFSNRTFactory.h
@@ -98,11 +98,11 @@ namespace Audio
 		static constexpr float DefaultShortTermLoudnessDuration = 3.f;
 
 		FLoudnessAnalyzerSettings LoudnessAnalyzerSettings = { 
-			.FFTSize = 2048,
-			.WindowType = EWindowType::Blackman,
+			.FFTSize = 4096,
+			.WindowType = EWindowType::None,
 			.LoudnessCurveType = ELoudnessCurveType::K,
 			.MinAnalysisFrequency = 20.f,
-			.MaxAnalysisFrequency = 20000.f,
+			.MaxAnalysisFrequency = 24000.f,
 			.ScalingMethod = ELoudnessAnalyzerScalingMethod::Corrected
 		};
 
diff --git a/Engine/Source/Runtime/SignalProcessing/Private/LoudnessAnalyzer.cpp b/Engine/Source/Runtime/SignalProcessing/Private/LoudnessAnalyzer.cpp
--- a/Engine/Source/Runtime/SignalProcessing/Private/LoudnessAnalyzer.cpp
+++ b/Engine/Source/Runtime/SignalProcessing/Private/LoudnessAnalyzer.cpp
@@ -47,10 +47,10 @@ namespace Audio
 		return (Freq / 6.8966888496476e-5f) * FMath::Sqrt(H / ((FreqSqrd + 79919.29f) * (FreqSqrd + 1345600.f)));
     }
 
-	float CalculatePolynomial(const float Freq, const float x0, const float x1, const float x2)
+	float CalculatePolynomial(const double Freq, const double x0, const double x1, const double x2)
     {
 	    const double CosFreq = FMath::Cos(Freq);
-    	return x1 * x1 + (x0 - x2) * (x0 - x2) + 2.f * x1 * (x0 + x2) * CosFreq + 4.f * x0 * x2 * CosFreq * CosFreq;
+    	return x1 * x1 + (x0 - x2) * (x0 - x2) + 2.0 * x1 * (x0 + x2) * CosFreq + 4.0 * x0 * x2 * CosFreq * CosFreq;
     }
 
 	float GetEqualLoudnessKWeightForFrequency(const float Freq)
diff --git a/Engine/Source/Runtime/SignalProcessing/Public/DSP/LKFSAnalyzer.h b/Engine/Source/Runtime/SignalProcessing/Public/DSP/LKFSAnalyzer.h
--- a/Engine/Source/Runtime/SignalProcessing/Public/DSP/LKFSAnalyzer.h
+++ b/Engine/Source/Runtime/SignalProcessing/Public/DSP/LKFSAnalyzer.h
@@ -24,8 +24,8 @@ namespace Audio
 		static constexpr float DefaultIntegratedLoudnessDuration = 60.f; // 1 minute
 
 		FLoudnessAnalyzerSettings LoudnessAnalyzerSettings {
-			.FFTSize = 2048,
-			.WindowType = EWindowType::Blackman,
+			.FFTSize = 4096,
+			.WindowType = EWindowType::None,
 			.LoudnessCurveType = ELoudnessCurveType::K,
 			.MinAnalysisFrequency = 20.0f,
 			.MaxAnalysisFrequency = 24000.0f,
