const correctness fix in ConvertDeinterleave interface

#rnx
#Synced-CL 51711552

diff --git a/Engine/Source/Runtime/Engine/Classes/Sound/SoundWaveProxyPlayer.h b/Engine/Source/Runtime/Engine/Classes/Sound/SoundWaveProxyPlayer.h
--- a/Engine/Source/Runtime/Engine/Classes/Sound/SoundWaveProxyPlayer.h
+++ b/Engine/Source/Runtime/Engine/Classes/Sound/SoundWaveProxyPlayer.h
@@ -152,7 +152,7 @@ public:
 	 * 
 	 * @param OutMultiChannelAudio 
 	 */
-	UE_API void RenderMultiChannelAudio(Audio::FMultichannelBufferView& OutMultiChannelAudio);
+	UE_API void RenderMultiChannelAudio(const Audio::FMultichannelBufferView& OutMultiChannelAudio);
 
 	/**
 	 * Render audio from the Wave Source based on the current Speed, Loop Settings, and Desired output Sample Rate
diff --git a/Engine/Source/Runtime/Engine/Private/SoundWaveProxyPlayer.cpp b/Engine/Source/Runtime/Engine/Private/SoundWaveProxyPlayer.cpp
--- a/Engine/Source/Runtime/Engine/Private/SoundWaveProxyPlayer.cpp
+++ b/Engine/Source/Runtime/Engine/Private/SoundWaveProxyPlayer.cpp
@@ -279,7 +279,7 @@ void FSoundWaveProxyPlayer::RenderMultiChannelAudio(int32 BlockFrameIndex, Audio
 	RenderMultiChannelAudio(OutMultiChannelAudio);
 }
 
-void FSoundWaveProxyPlayer::RenderMultiChannelAudio(Audio::FMultichannelBufferView& OutMultiChannelAudio)
+void FSoundWaveProxyPlayer::RenderMultiChannelAudio(const Audio::FMultichannelBufferView& OutMultiChannelAudio)
 {
 	check(SoundWave.IsValid());
 	
@@ -290,7 +290,7 @@ void FSoundWaveProxyPlayer::RenderMultiChannelAudio(Audio::FMultichannelBufferVi
 	
 	if (IsFinished())
 	{
-		for (TArrayView<float>& Audio : OutMultiChannelAudio)
+		for (const TArrayView<float>& Audio : OutMultiChannelAudio)
 		{
 			FMemory::Memset(Audio.GetData(), 0, Audio.Num() * Audio.GetTypeSize());
 		}
diff --git a/Engine/Source/Runtime/SignalProcessing/Private/ConvertDeinterleave.cpp b/Engine/Source/Runtime/SignalProcessing/Private/ConvertDeinterleave.cpp
--- a/Engine/Source/Runtime/SignalProcessing/Private/ConvertDeinterleave.cpp
+++ b/Engine/Source/Runtime/SignalProcessing/Private/ConvertDeinterleave.cpp
@@ -52,7 +52,7 @@ namespace Audio
 				return &ChannelGains.GetData()[InOutputChannelIndex * NumInputChannels];
 			}
 
-			virtual void ProcessAudio(TArrayView<const float> InSamples, TArray<TArrayView<float>>& OutMultiChannelAudio) const override
+			virtual void ProcessAudio(TArrayView<const float> InSamples, const TArray<TArrayView<float>>& OutMultiChannelAudio) const override
 			{
 				checkf((InSamples.Num() % NumInputChannels) == 0, TEXT("Input sample buffer contains partial audio frame."));
 				checkf((InSamples.Num() / NumInputChannels) == GetMultichannelBufferNumFrames(OutMultiChannelAudio), TEXT("Audio buffer frame count mismatch."));
@@ -120,7 +120,7 @@ namespace Audio
 
 			virtual ~TConvertDeinterleave() = default;
 
-			virtual void ProcessAudio(TArrayView<const float> InSamples, TArray<TArrayView<float>>& OutMultiChannelAudio) const override
+			virtual void ProcessAudio(TArrayView<const float> InSamples, const TArray<TArrayView<float>>& OutMultiChannelAudio) const override
 			{
 				checkf(InSamples.Num() == GetMultichannelBufferNumFrames(OutMultiChannelAudio), TEXT("Audio buffer frame count mismatch."));
 				checkf(NumOutputChannels == OutMultiChannelAudio.Num(), TEXT("Output audio buffer not initialized to expected channel count."));
@@ -164,7 +164,7 @@ namespace Audio
 			}
 
 			virtual ~TConvertDeinterleave() = default;
-			virtual void ProcessAudio(TArrayView<const float> InSamples, TArray<TArrayView<float>>& OutMultiChannelAudio) const override
+			virtual void ProcessAudio(TArrayView<const float> InSamples, const TArray<TArrayView<float>>& OutMultiChannelAudio) const override
 			{
 				checkf(InSamples.Num() == GetMultichannelBufferNumFrames(OutMultiChannelAudio), TEXT("Audio buffer frame count mismatch."));
 				checkf(1 == OutMultiChannelAudio.Num(), TEXT("Output audio buffer not initialized to expected channel count."));
diff --git a/Engine/Source/Runtime/SignalProcessing/Public/DSP/ConvertDeinterleave.h b/Engine/Source/Runtime/SignalProcessing/Public/DSP/ConvertDeinterleave.h
--- a/Engine/Source/Runtime/SignalProcessing/Public/DSP/ConvertDeinterleave.h
+++ b/Engine/Source/Runtime/SignalProcessing/Public/DSP/ConvertDeinterleave.h
@@ -58,7 +58,7 @@ namespace Audio
 		 *                     the input samples and possibly upmixing or downmixing
 		 *                     them to the target channel count. 
 		 */
-		virtual void ProcessAudio(TArrayView<const float> InSamples, TArray<TArrayView<float>>& OutMultiChannelAudio) const = 0;
+		virtual void ProcessAudio(TArrayView<const float> InSamples, const TArray<TArrayView<float>>& OutMultiChannelAudio) const = 0;
 
 
 		/** Create an IConvertDeinterleave object for a given number of input
