ATR, adding lpf/hpf to submix and audio bus sends

Adds optional per-send LPF and HPF filtering to both submix sends and audio bus sends.
When enabled, the source audio is filtered through a one-pole LPF and HPF in series before
being mixed into the destination.

#Synced-CL 52423952

diff --git a/Engine/Source/Editor/DetailCustomizations/Private/DetailCustomizations.cpp b/Engine/Source/Editor/DetailCustomizations/Private/DetailCustomizations.cpp
--- a/Engine/Source/Editor/DetailCustomizations/Private/DetailCustomizations.cpp
+++ b/Engine/Source/Editor/DetailCustomizations/Private/DetailCustomizations.cpp
@@ -141,6 +141,7 @@
 #include "SkyLightComponentDetails.h"
 #include "SlateColorCustomization.h"
 #include "SlateSoundCustomization.h"
+#include "SoundSubmixSendFilterCustomization.h"
 #include "SoftClassPathCustomization.h"
 #include "SoftObjectPathCustomization.h"
 #include "Sound/SoundNodeDistanceCrossFade.h"
@@ -306,6 +307,8 @@ void FDetailCustomizationsModule::RegisterPropertyTypeCustomizations()
 	RegisterCustomPropertyTypeLayout("SlateColor", FOnGetPropertyTypeCustomizationInstance::CreateStatic(&FSlateColorCustomization::MakeInstance));
 	RegisterCustomPropertyTypeLayout("ForceFeedbackAttenuationSettings", FOnGetPropertyTypeCustomizationInstance::CreateStatic(&FForceFeedbackAttenuationSettingsCustomization::MakeInstance));
 	RegisterCustomPropertyTypeLayout("SoundAttenuationSettings", FOnGetPropertyTypeCustomizationInstance::CreateStatic(&FSoundAttenuationSettingsCustomization::MakeInstance));
+	RegisterCustomPropertyTypeLayout("SoundSubmixSendInfo", FOnGetPropertyTypeCustomizationInstance::CreateStatic(&FSoundSubmixSendFilterCustomization::MakeInstance));
+	RegisterCustomPropertyTypeLayout("SoundSourceBusSendInfo", FOnGetPropertyTypeCustomizationInstance::CreateStatic(&FSoundSubmixSendFilterCustomization::MakeInstance));
 	RegisterCustomPropertyTypeLayout("DialogueContext", FOnGetPropertyTypeCustomizationInstance::CreateStatic(&FDialogueContextStructCustomization::MakeInstance));
 	RegisterCustomPropertyTypeLayout("DialogueWaveParameter", FOnGetPropertyTypeCustomizationInstance::CreateStatic(&FDialogueWaveParameterStructCustomization::MakeInstance));
 	RegisterCustomPropertyTypeLayout("BodyInstance", FOnGetPropertyTypeCustomizationInstance::CreateStatic(&FBodyInstanceCustomization::MakeInstance));
diff --git a/Engine/Source/Editor/DetailCustomizations/Private/SoundSubmixSendFilterCustomization.cpp b/Engine/Source/Editor/DetailCustomizations/Private/SoundSubmixSendFilterCustomization.cpp
new file mode 100644
--- /dev/null
+++ b/Engine/Source/Editor/DetailCustomizations/Private/SoundSubmixSendFilterCustomization.cpp
@@ -0,0 +1,71 @@
+// Copyright Epic Games, Inc. All Rights Reserved.
+
+#include "SoundSubmixSendFilterCustomization.h"
+
+#include "DetailWidgetRow.h"
+#include "IDetailChildrenBuilder.h"
+#include "PropertyHandle.h"
+#include "UObject/PropertyOptional.h"
+#include "UObject/UnrealType.h"
+
+TSharedRef<IPropertyTypeCustomization> FSoundSubmixSendFilterCustomization::MakeInstance()
+{
+	return MakeShareable(new FSoundSubmixSendFilterCustomization);
+}
+
+void FSoundSubmixSendFilterCustomization::CustomizeHeader(
+	TSharedRef<IPropertyHandle> StructHandle,
+	FDetailWidgetRow& HeaderRow,
+	IPropertyTypeCustomizationUtils& Utils)
+{
+	HeaderRow
+		.NameContent()  [ StructHandle->CreatePropertyNameWidget()  ]
+		.ValueContent() [ StructHandle->CreatePropertyValueWidget() ];
+}
+
+void FSoundSubmixSendFilterCustomization::CustomizeChildren(
+	TSharedRef<IPropertyHandle> StructHandle,
+	IDetailChildrenBuilder& ChildBuilder,
+	IPropertyTypeCustomizationUtils& Utils)
+{
+	uint32 NumChildren = 0;
+	StructHandle->GetNumChildren(NumChildren);
+
+	for (uint32 i = 0; i < NumChildren; ++i)
+	{
+		TSharedPtr<IPropertyHandle> Child = StructHandle->GetChildHandle(i);
+		if (!Child.IsValid())
+		{
+			continue;
+		}
+
+		const FName PropertyName = Child->GetProperty()->GetFName();
+		if (PropertyName == FName("LPFCutoffFrequency"))
+		{
+			// Stamp metadata onto the inner float property so the spinbox shows "Hz",
+			// steps by 0.1 Hz, and rejects values below 20 Hz.
+			// UHT rejects these on TOptional<float>, so we inject them here instead.
+			if (FOptionalProperty* OptProp = CastField<FOptionalProperty>(Child->GetProperty()))
+			{
+				OptProp->GetValueProperty()->SetMetaData(TEXT("Units"), TEXT("Hz"));
+				OptProp->GetValueProperty()->SetMetaData(TEXT("Delta"), TEXT("0.1"));
+				OptProp->GetValueProperty()->SetMetaData(TEXT("ClampMin"), TEXT("20.0"));
+			}
+			ChildBuilder.AddProperty(Child.ToSharedRef());
+		}
+		else if (PropertyName == FName("HPFCutoffFrequency"))
+		{
+			if (FOptionalProperty* OptProp = CastField<FOptionalProperty>(Child->GetProperty()))
+			{
+				OptProp->GetValueProperty()->SetMetaData(TEXT("Units"), TEXT("Hz"));
+				OptProp->GetValueProperty()->SetMetaData(TEXT("Delta"), TEXT("0.1"));
+				OptProp->GetValueProperty()->SetMetaData(TEXT("ClampMin"), TEXT("20.0"));
+			}
+			ChildBuilder.AddProperty(Child.ToSharedRef());
+		}
+		else
+		{
+			ChildBuilder.AddProperty(Child.ToSharedRef());
+		}
+	}
+}
diff --git a/Engine/Source/Editor/DetailCustomizations/Private/SoundSubmixSendFilterCustomization.h b/Engine/Source/Editor/DetailCustomizations/Private/SoundSubmixSendFilterCustomization.h
new file mode 100644
--- /dev/null
+++ b/Engine/Source/Editor/DetailCustomizations/Private/SoundSubmixSendFilterCustomization.h
@@ -0,0 +1,27 @@
+// Copyright Epic Games, Inc. All Rights Reserved.
+#pragma once
+
+#include "IPropertyTypeCustomization.h"
+#include "Templates/SharedPointer.h"
+
+class FDetailWidgetRow;
+class IDetailChildrenBuilder;
+class IPropertyHandle;
+class IPropertyTypeCustomizationUtils;
+
+/**
+ * Detail customization for FSoundSubmixSendInfo and FSoundSourceBusSendInfo.
+ * Injects Units, Delta, and ClampMin metadata onto the inner float property of the
+ * TOptional<float> LPF/HPF cutoff frequency UPROPERTYs so the editor spinbox shows
+ * Hz units, steps by 0.1 Hz, and rejects values below 20 Hz.
+ */
+class FSoundSubmixSendFilterCustomization : public IPropertyTypeCustomization
+{
+public:
+	static TSharedRef<IPropertyTypeCustomization> MakeInstance();
+
+	//~ Begin IPropertyTypeCustomization
+	virtual void CustomizeHeader(TSharedRef<IPropertyHandle> StructHandle, FDetailWidgetRow& HeaderRow, IPropertyTypeCustomizationUtils& Utils) override;
+	virtual void CustomizeChildren(TSharedRef<IPropertyHandle> StructHandle, IDetailChildrenBuilder& ChildBuilder, IPropertyTypeCustomizationUtils& Utils) override;
+	//~ End IPropertyTypeCustomization
+};
diff --git a/Engine/Source/Runtime/AudioMixer/Private/AudioMixerBus.cpp b/Engine/Source/Runtime/AudioMixer/Private/AudioMixerBus.cpp
--- a/Engine/Source/Runtime/AudioMixer/Private/AudioMixerBus.cpp
+++ b/Engine/Source/Runtime/AudioMixer/Private/AudioMixerBus.cpp
@@ -83,10 +83,30 @@ namespace Audio
 		// Make sure we don't have duplicates in the bus sends
 		for (FAudioBusSend& BusSend : AudioBusSends[(int32)BusSendType])
 		{
-			// If it's already added, just update the send level
+			// If it's already added, update send level and filter state
 			if (BusSend.SourceId == InAudioBusSend.SourceId)
 			{
 				BusSend.SendLevel = InAudioBusSend.SendLevel;
+					if (InAudioBusSend.Filter.bEnableLPF)
+					{
+						BusSend.Filter.LpfCutoff = InAudioBusSend.Filter.LpfCutoff;
+						// If LPF was previously disabled, the Lpf object was never Init()'d; copy it in.
+						if (!BusSend.Filter.bEnableLPF)
+						{
+							BusSend.Lpf = InAudioBusSend.Lpf;
+						}
+					}
+					if (InAudioBusSend.Filter.bEnableHPF)
+					{
+						BusSend.Filter.HpfCutoff = InAudioBusSend.Filter.HpfCutoff;
+						// If HPF was previously disabled, the Hpf object was never Init()'d; copy it in.
+						if (!BusSend.Filter.bEnableHPF)
+						{
+							BusSend.Hpf = InAudioBusSend.Hpf;
+						}
+					}
+					BusSend.Filter.bEnableLPF = InAudioBusSend.Filter.bEnableLPF;
+					BusSend.Filter.bEnableHPF = InAudioBusSend.Filter.bEnableHPF;
 				return;
 			}
 		}
@@ -146,10 +166,15 @@ namespace Audio
 
 		PatchMixer.PopAudio(BusDataBufferPtr, NumSamples, false);
 
+		// FilterScratchBuffer is a member so its capacity persists across render callbacks,
+		// avoiding per-frame heap reallocation once the working size is established.
+		// Reset element count; actual size is set on demand inside the send loop.
+		FilterScratchBuffer.Reset();
+
 		for (int32 BusSendType = 0; BusSendType < (int32)EBusSendType::Count; ++BusSendType)
 		{
 			// Loop through the send list for this bus
-			for (const FAudioBusSend& AudioBusSend : AudioBusSends[BusSendType])
+				for (FAudioBusSend& AudioBusSend : AudioBusSends[BusSendType])
 			{
 				const float* SourceBufferPtr = nullptr;
 
@@ -184,16 +209,62 @@ namespace Audio
 					// Up-mix or down-mix if source channels differ from bus channels
 					if (NumSourceChannels != NumChannels)
 					{
+							const float* FilteredSourcePtr = SourceBufferPtr;
+
+							// Apply per-send filters to the source buffer before downmixing.
+							// When both LPF and HPF are active, the LPF writes into FilterScratchBuffer and
+							// the HPF processes it via ProcessBufferInPlace -- avoiding aliased pointers into
+							// ProcessAudioBuffer (whose RESTRICT arguments must not alias).
+							if (AudioBusSend.Filter.IsFiltering())
+							{
+								const bool bBypassLPF = !AudioBusSend.Filter.bEnableLPF;
+								const bool bBypassHPF = !AudioBusSend.Filter.bEnableHPF;
+
+								FilterScratchBuffer.SetNumUninitialized(NumSourceSamples);
+
+								if (!bBypassLPF)
+								{
+									AudioBusSend.Lpf.StartFrequencyInterpolation(AudioBusSend.Filter.LpfCutoff, NumOutputFrames);
+									AudioBusSend.Lpf.ProcessAudioBuffer(SourceBufferPtr, FilterScratchBuffer.GetData(), NumSourceSamples);
+									AudioBusSend.Lpf.StopFrequencyInterpolation();
+									FilteredSourcePtr = FilterScratchBuffer.GetData();
+								}
+
+								if (!bBypassHPF)
+								{
+									AudioBusSend.Hpf.StartFrequencyInterpolation(AudioBusSend.Filter.HpfCutoff, NumOutputFrames);
+									if (FilteredSourcePtr == FilterScratchBuffer.GetData())
+									{
+										// LPF already wrote into scratch; process HPF in-place using the dedicated API.
+										AudioBusSend.Hpf.ProcessBufferInPlace(FilterScratchBuffer.GetData(), NumSourceSamples);
+									}
+									else
+									{
+										AudioBusSend.Hpf.ProcessAudioBuffer(FilteredSourcePtr, FilterScratchBuffer.GetData(), NumSourceSamples);
+									}
+									AudioBusSend.Hpf.StopFrequencyInterpolation();
+									FilteredSourcePtr = FilterScratchBuffer.GetData();
+								}
+							}
+
 						FAlignedFloatBuffer ChannelMap;
 						SourceManager->Get2DChannelMap(AudioBusSend.SourceId, NumChannels, ChannelMap);
 						Algo::ForEach(ChannelMap, [SendLevel = AudioBusSend.SendLevel](float& ChannelValue) { ChannelValue *= SendLevel; });
-						DownmixAndSumIntoBuffer(NumSourceChannels, NumChannels, SourceBufferPtr, BusDataBufferPtr, NumOutputFrames, ChannelMap.GetData());
+							DownmixAndSumIntoBuffer(NumSourceChannels, NumChannels, FilteredSourcePtr, BusDataBufferPtr, NumOutputFrames, ChannelMap.GetData());
 					}
 					else
 					{
 						TArrayView<const float> SourceBufferView(SourceBufferPtr, NumOutputFrames * NumChannels);
 						TArrayView<float> BusDataBufferView(BusDataBufferPtr, NumOutputFrames * NumChannels);
-						ArrayMixIn(SourceBufferView, BusDataBufferView, AudioBusSend.SendLevel);
+
+						if (AudioBusSend.Filter.IsFiltering())
+						{
+							ArrayFilterAndMixIn(SourceBufferView, BusDataBufferView, AudioBusSend.SendLevel, NumChannels, AudioBusSend.Filter.bEnableLPF, AudioBusSend.Filter.LpfCutoff, AudioBusSend.Filter.bEnableHPF, AudioBusSend.Filter.HpfCutoff, AudioBusSend.Lpf, AudioBusSend.Hpf, FilterScratchBuffer);
+						}
+						else
+						{
+							ArrayMixIn(SourceBufferView, BusDataBufferView, AudioBusSend.SendLevel);
+						}
 					}
 				}
 			}
diff --git a/Engine/Source/Runtime/AudioMixer/Private/AudioMixerSource.cpp b/Engine/Source/Runtime/AudioMixer/Private/AudioMixerSource.cpp
--- a/Engine/Source/Runtime/AudioMixer/Private/AudioMixerSource.cpp
+++ b/Engine/Source/Runtime/AudioMixer/Private/AudioMixerSource.cpp
@@ -560,6 +560,8 @@ namespace Audio
 					
 					SubmixSend.bIsMainSend = false;
 					SubmixSend.SoundfieldFactory = MixerDevice->GetFactoryForSubmixInstance(SubmixSend.Submix);
+					SubmixSend.Filter = FSendFilterSettings(SendInfo.LPFCutoffFrequency.IsSet(), SendInfo.HPFCutoffFrequency.IsSet(),
+						SendInfo.LPFCutoffFrequency.Get(20000.0f), SendInfo.HPFCutoffFrequency.Get(20.0f));
 					InitParams.SubmixSends.Add(SubmixSend);
 				}
 			}
@@ -705,7 +707,10 @@ namespace Audio
 					{
 						BusSend.SendLevel = 0;
 					}
-					
+
+					BusSend.Filter = FSendFilterSettings(InSendInfo.LPFCutoffFrequency.IsSet(), InSendInfo.HPFCutoffFrequency.IsSet(),
+						InSendInfo.LPFCutoffFrequency.Get(20000.0f), InSendInfo.HPFCutoffFrequency.Get(20.0f));
+
 					if (AudioBusSends)
 					{
 						AudioBusSends[InBusSendType].Add(BusSend);
@@ -724,6 +729,8 @@ namespace Audio
 					NewDynamicBusSendInfo.MinSendDistance = InSendInfo.MinSendDistance;
 					NewDynamicBusSendInfo.MaxSendDistance = InSendInfo.MaxSendDistance;
 					NewDynamicBusSendInfo.CustomSendLevelCurve = InSendInfo.CustomSendLevelCurve;
+					NewDynamicBusSendInfo.Filter = FSendFilterSettings(InSendInfo.LPFCutoffFrequency.IsSet(), InSendInfo.HPFCutoffFrequency.IsSet(),
+						InSendInfo.LPFCutoffFrequency.Get(20000.0f), InSendInfo.HPFCutoffFrequency.Get(20.0f));
 
 					// Copy the bus SourceBusSendInfo structs to a local copy so we can update it in the update tick
 					bool bIsNew = true;
@@ -731,7 +738,10 @@ namespace Audio
 					{
 						if (BusSendInfo.BusId == NewDynamicBusSendInfo.BusId)
 						{
+							// Detect filter param changes before overwriting so the dirty-check can dispatch them.
+							const bool bFilteringChanged = BusSendInfo.Filter.HasChanged(NewDynamicBusSendInfo.Filter);
 							BusSendInfo = NewDynamicBusSendInfo;
+							BusSendInfo.bFilterStateChanged |= bFilteringChanged;
 							BusSendInfo.bIsInit = false;
 							bIsNew = false;
 							break;
@@ -1591,7 +1601,9 @@ namespace Audio
 			}
 
 			// set the level and stage for this send
-			MixerSourceVoice->SetSubmixSendInfo(SubmixInstance, SendLevel, InSendStage);
+			MixerSourceVoice->SetSubmixSendInfo(SubmixInstance, SendLevel, InSendStage,
+				FSendFilterSettings(InSendInfo.LPFCutoffFrequency.IsSet(), InSendInfo.HPFCutoffFrequency.IsSet(),
+					InSendInfo.LPFCutoffFrequency.Get(20000.0f), InSendInfo.HPFCutoffFrequency.Get(20.0f)));
 		}
 	}
 
@@ -1766,8 +1778,26 @@ namespace Audio
 					BusName = NewSend.Value.SoundSourceBus->GetPathName();
 #endif // if ENABLE_AUDIO_DEBUG
 
-					MixerSourceVoice->SetAudioBusSendInfo(NewSend.Key, NewSend.Value.SoundSourceBus->GetUniqueID(), NewSend.Value.SendLevel, BusName);
+					// Match SetupBusData: prefer the underlying AudioBus ID when present (explicit bus),
+					// otherwise fall back to the SoundSourceBus own ID (automatic source bus).
+					const uint32 BusId = NewSend.Value.SoundSourceBus->AudioBus
+						? NewSend.Value.SoundSourceBus->AudioBus->GetUniqueID()
+						: NewSend.Value.SoundSourceBus->GetUniqueID();
+					MixerSourceVoice->SetAudioBusSendInfo(NewSend.Key, BusId, NewSend.Value.SendLevel, BusName,
+						FSendFilterSettings(NewSend.Value.LPFCutoffFrequency.IsSet(), NewSend.Value.HPFCutoffFrequency.IsSet(),
+							NewSend.Value.LPFCutoffFrequency.Get(20000.0f), NewSend.Value.HPFCutoffFrequency.Get(20.0f)));
 					bSendingAudioToBuses = true;
+					// Sync DynamicBusSendInfos so the per-tick dirty-check does not double-dispatch this send.
+					for (FDynamicBusSendInfo& DynamicInfo : DynamicBusSendInfos)
+					{
+						if (DynamicInfo.BusId == BusId)
+						{
+							DynamicInfo.SendLevel = NewSend.Value.SendLevel;
+							DynamicInfo.Filter = FSendFilterSettings(NewSend.Value.LPFCutoffFrequency.IsSet(), NewSend.Value.HPFCutoffFrequency.IsSet(),
+								NewSend.Value.LPFCutoffFrequency.Get(20000.0f), NewSend.Value.HPFCutoffFrequency.Get(20.0f));
+							DynamicInfo.bFilterStateChanged = false;
+						}
+					}
 				}
 
 				if (NewSend.Value.AudioBus)
@@ -1778,8 +1808,22 @@ namespace Audio
 					BusName = NewSend.Value.AudioBus->GetPathName();
 #endif // if ENABLE_AUDIO_DEBUG
 
-					MixerSourceVoice->SetAudioBusSendInfo(NewSend.Key, NewSend.Value.AudioBus->GetUniqueID(), NewSend.Value.SendLevel, BusName);
+					const uint32 AudioBusId = NewSend.Value.AudioBus->GetUniqueID();
+					MixerSourceVoice->SetAudioBusSendInfo(NewSend.Key, AudioBusId, NewSend.Value.SendLevel, BusName,
+						FSendFilterSettings(NewSend.Value.LPFCutoffFrequency.IsSet(), NewSend.Value.HPFCutoffFrequency.IsSet(),
+							NewSend.Value.LPFCutoffFrequency.Get(20000.0f), NewSend.Value.HPFCutoffFrequency.Get(20.0f)));
 					bSendingAudioToBuses = true;
+					// Sync DynamicBusSendInfos so the per-tick dirty-check does not double-dispatch this send.
+					for (FDynamicBusSendInfo& DynamicInfo : DynamicBusSendInfos)
+					{
+						if (DynamicInfo.BusId == AudioBusId)
+						{
+							DynamicInfo.SendLevel = NewSend.Value.SendLevel;
+							DynamicInfo.Filter = FSendFilterSettings(NewSend.Value.LPFCutoffFrequency.IsSet(), NewSend.Value.HPFCutoffFrequency.IsSet(),
+								NewSend.Value.LPFCutoffFrequency.Get(20000.0f), NewSend.Value.HPFCutoffFrequency.Get(20.0f));
+							DynamicInfo.bFilterStateChanged = false;
+						}
+					}
 				}
 			}
 
@@ -1816,11 +1860,13 @@ namespace Audio
 			// If the send level changed, then we need to send an update to the audio render thread
 			const bool bSendLevelChanged = !FMath::IsNearlyEqual(SendLevel, DynamicBusSendInfo.SendLevel);
 			const bool bBusEnablementChanged = bPreviousBusEnablement != WaveInstance->bEnableBusSends;
+			const bool bFilterStateChanged = DynamicBusSendInfo.bFilterStateChanged;
 
-			if (bSendLevelChanged || bBusEnablementChanged)
+			if (bSendLevelChanged || bBusEnablementChanged || bFilterStateChanged)
 			{
 				DynamicBusSendInfo.SendLevel = SendLevel;
 				DynamicBusSendInfo.bIsInit = false;
+				DynamicBusSendInfo.bFilterStateChanged = false;
 
 				FString BusName;
 
@@ -1828,7 +1874,8 @@ namespace Audio
 				BusName = DynamicBusSendInfo.BusName;
 #endif // if ENABLE_AUDIO_DEBUG
 
-				MixerSourceVoice->SetAudioBusSendInfo(DynamicBusSendInfo.BusSendType, DynamicBusSendInfo.BusId, SendLevel, BusName);
+				MixerSourceVoice->SetAudioBusSendInfo(DynamicBusSendInfo.BusSendType, DynamicBusSendInfo.BusId, SendLevel, BusName,
+					DynamicBusSendInfo.Filter);
 
 				bPreviousBusEnablement = WaveInstance->bEnableBusSends;
 			}
diff --git a/Engine/Source/Runtime/AudioMixer/Private/AudioMixerSource.h b/Engine/Source/Runtime/AudioMixer/Private/AudioMixerSource.h
--- a/Engine/Source/Runtime/AudioMixer/Private/AudioMixerSource.h
+++ b/Engine/Source/Runtime/AudioMixer/Private/AudioMixerSource.h
@@ -170,6 +170,10 @@ namespace Audio
 			float MaxSendDistance = 0.0f;
 			FRuntimeFloatCurve CustomSendLevelCurve;
 			bool bIsInit = true;
+			FSendFilterSettings Filter;
+			// Set when filter params change via the bHasNewBusSends path so the dynamic
+			// tick loop re-sends the update even if send level didn't change.
+			bool bFilterStateChanged = false;
 		};
 
 		// Mapping of channel map types to channel maps. Determined by what submixes this source sends its audio to.
diff --git a/Engine/Source/Runtime/AudioMixer/Private/AudioMixerSourceManager.cpp b/Engine/Source/Runtime/AudioMixer/Private/AudioMixerSourceManager.cpp
--- a/Engine/Source/Runtime/AudioMixer/Private/AudioMixerSourceManager.cpp
+++ b/Engine/Source/Runtime/AudioMixer/Private/AudioMixerSourceManager.cpp
@@ -1466,11 +1466,20 @@ namespace Audio
 					NewAudioBusSend.WaveInstancePlayOrder = SourceInfo.PlayOrder;
 					NewAudioBusSend.SendLevel = AudioBusSend.SendLevel;
 
-					// Get existing BusId and add the send, or create new bus registration. A newly created bus won't output audio, but 
+					// Get existing BusId and add the send, or create new bus registration. A newly created bus won't output audio, but
 					// we register the send anyway in the event that this bus does play, we'll know to send this
 					// source's audio to it.
 					bool bAlreadyExisted = false;
 					TSharedRef<FMixerAudioBus> AudioBusPtr = FindOrAddAudioBus(AudioBusSend.AudioBusId, AudioBusSend.BusChannels, true, &bAlreadyExisted);
+
+					if (AudioBusSend.Filter.IsFiltering())
+					{
+						NewAudioBusSend.Filter = AudioBusSend.Filter;
+						const int32 NumSourceChannels = FMath::Max(SourceInfo.NumInputChannels, 1);
+						NewAudioBusSend.Lpf.Init(MixerDevice->SampleRate, NumSourceChannels);
+						NewAudioBusSend.Hpf.Init(MixerDevice->SampleRate, NumSourceChannels);
+					}
+
 					AudioBusPtr->AddSend((EBusSendType)BusSendType, NewAudioBusSend);
 #if UE_AUDIO_PROFILERTRACE_ENABLED
 					if (!bAlreadyExisted)
@@ -1506,7 +1515,30 @@ namespace Audio
 				FMixerSubmixPtr SubmixPtr = MixerSubmixSend.Submix.Pin();
 				if (SubmixPtr.IsValid())
 				{
-					SourceInfo.SubmixSends.Add(MixerSubmixSend);
+					FMixerSourceSubmixSend& AddedSend = SourceInfo.SubmixSends.Add_GetRef(MixerSubmixSend);
+					AddedSend.SubmixId = SubmixPtr->GetId();
+
+					if (AddedSend.Filter.IsFiltering())
+					{
+						// Use device output channel count: the filtered buffer is the source device-format output
+						// buffer, not the destination submix output channel count.
+						const int32 NumDeviceChannels = FMath::Max(MixerDevice->GetNumDeviceChannels(), 1);
+						AddedSend.Lpf.Init(MixerDevice->SampleRate, NumDeviceChannels);
+						AddedSend.Hpf.Init(MixerDevice->SampleRate, NumDeviceChannels);
+						
+						// Snap coefficients to the initial cutoffs immediately. bIsFirstFrequencyChange (reset by
+						// Init) forces CurrInterpLength=0 in StartFrequencyInterpolation, so coefficients land on
+						// the target without interpolation rather than ramping from the default.
+						if (AddedSend.Filter.bEnableLPF)
+						{
+							AddedSend.Lpf.StartFrequencyInterpolation(AddedSend.Filter.LpfCutoff, 1);
+						}
+						
+						if (AddedSend.Filter.bEnableHPF)
+						{
+							AddedSend.Hpf.StartFrequencyInterpolation(AddedSend.Filter.HpfCutoff, 1);
+						}
+					}
 
 					if (MixerSubmixSend.SubmixSendStage == EMixerSourceSubmixSendStage::PreDistanceAttenuation)
 					{
@@ -2171,6 +2203,31 @@ namespace Audio
 							// Update existing submix send if it already exists
 							SubmixSend.SendLevel = InSubmixSend.SendLevel;
 							SubmixSend.SubmixSendStage = InSubmixSend.SubmixSendStage;
+							const bool bWasFilteringEnabled = SubmixSend.Filter.IsFiltering();
+							SubmixSend.Filter.bEnableLPF = InSubmixSend.Filter.bEnableLPF;
+							SubmixSend.Filter.bEnableHPF = InSubmixSend.Filter.bEnableHPF;
+							if (SubmixSend.Filter.IsFiltering())
+							{
+								if (!bWasFilteringEnabled)
+								{
+									const int32 NumDeviceChannels = FMath::Max(MixerDevice->GetNumDeviceChannels(), 1);
+									SubmixSend.Lpf.Init(MixerDevice->SampleRate, NumDeviceChannels);
+									SubmixSend.Hpf.Init(MixerDevice->SampleRate, NumDeviceChannels);
+								}
+								
+								if (SubmixSend.Filter.bEnableLPF)
+								{
+									SubmixSend.Filter.LpfCutoff = InSubmixSend.Filter.LpfCutoff;
+									SubmixSend.Lpf.StartFrequencyInterpolation(SubmixSend.Filter.LpfCutoff, 1);
+								}
+								
+								if (SubmixSend.Filter.bEnableHPF)
+								{
+									SubmixSend.Filter.HpfCutoff = InSubmixSend.Filter.HpfCutoff;
+									SubmixSend.Hpf.StartFrequencyInterpolation(SubmixSend.Filter.HpfCutoff, 1);
+								}
+							}
+							
 							bIsNew = false;
 							if (SourceInfo.bHasPreDistanceAttenuationSend)
 							{
@@ -2187,7 +2244,28 @@ namespace Audio
 
 				if (bIsNew)
 				{
-					SourceInfo.SubmixSends.Add(InSubmixSend);
+					FMixerSourceSubmixSend NewSend = InSubmixSend;
+					NewSend.SubmixId = InSubmixPtr->GetId();
+					if (NewSend.Filter.IsFiltering())
+					{
+						const int32 NumDeviceChannels = FMath::Max(MixerDevice->GetNumDeviceChannels(), 1);
+						NewSend.Lpf.Init(MixerDevice->SampleRate, NumDeviceChannels);
+						NewSend.Hpf.Init(MixerDevice->SampleRate, NumDeviceChannels);
+						// Snap coefficients to the initial cutoffs immediately. bIsFirstFrequencyChange (reset by
+						// Init) forces CurrInterpLength=0 in StartFrequencyInterpolation, so coefficients land on
+						// the target without interpolation rather than ramping from the default.
+						if (NewSend.Filter.bEnableLPF)
+						{
+							NewSend.Lpf.StartFrequencyInterpolation(NewSend.Filter.LpfCutoff, 1);
+						}
+						
+						if (NewSend.Filter.bEnableHPF)
+						{
+							NewSend.Hpf.StartFrequencyInterpolation(NewSend.Filter.HpfCutoff, 1);
+						}
+					}
+					
+					SourceInfo.SubmixSends.Add(MoveTemp(NewSend));
 				}
 				
 				// If we don't have a pre-distance attenuation send, lets zero out the buffer so the output buffer stops doing math with it.
@@ -2249,13 +2327,13 @@ namespace Audio
 		}, AUDIO_MIXER_THREAD_COMMAND_STRING("ClearSubmixSendInfo()"));
 	}
 
-	void FMixerSourceManager::SetBusSendInfo(const int32 SourceId, EBusSendType InAudioBusSendType, uint32 AudioBusId, float BusSendLevel, const FString& InBusName)
+	void FMixerSourceManager::SetBusSendInfo(const int32 SourceId, EBusSendType InAudioBusSendType, uint32 AudioBusId, float BusSendLevel, const FString& InBusName, const FSendFilterSettings& InFilter)
 	{
 		AUDIO_MIXER_CHECK(SourceId < NumLegacySources);
 		AUDIO_MIXER_CHECK(GameThreadInfo.bIsBusy[SourceId]);
 		AUDIO_MIXER_CHECK_GAME_THREAD(MixerDevice);
 
-		AudioMixerThreadCommand([this, SourceId, InAudioBusSendType, AudioBusId, BusSendLevel, InBusName]()
+		AudioMixerThreadCommand([this, SourceId, InAudioBusSendType, AudioBusId, BusSendLevel, InBusName, InFilter]()
 		{
 			FSourceInfo& SourceInfo = SourceInfos[SourceId];
 
@@ -2265,11 +2343,20 @@ namespace Audio
 			BusSend.WaveInstancePlayOrder = SourceInfo.PlayOrder;
 			BusSend.SendLevel = BusSendLevel;
 
-			// Retrieve the bus we want to send audio to. A newly created bus won't output audio, but 
+			// Retrieve the bus we want to send audio to. A newly created bus won't output audio, but
 			// we register the send anyway in the event that this bus does play, we'll know to send this
 			// source's audio to it.
 			bool bAlreadyExisted = false;
 			TSharedRef<FMixerAudioBus> AudioBusPtr = FindOrAddAudioBus(AudioBusId, SourceInfo.NumInputChannels, true, &bAlreadyExisted);
+
+			BusSend.Filter = InFilter;
+			if (InFilter.IsFiltering())
+			{
+				const int32 NumSourceChannels = FMath::Max(SourceInfo.NumInputChannels, 1);
+				BusSend.Lpf.Init(MixerDevice->SampleRate, NumSourceChannels);
+				BusSend.Hpf.Init(MixerDevice->SampleRate, NumSourceChannels);
+			}
+
 			AudioBusPtr->AddSend(InAudioBusSendType, BusSend);
 
 #if UE_AUDIO_PROFILERTRACE_ENABLED
@@ -3463,17 +3550,37 @@ namespace Audio
 		}
 	}
 
-	void FMixerSourceManager::MixOutputBuffers(const int32 SourceId, int32 InNumOutputChannels, const float InSendLevel, EMixerSourceSubmixSendStage InSubmixSendStage, FAlignedFloatBuffer& OutWetBuffer) const
+	void FMixerSourceManager::MixOutputBuffers(const int32 SourceId, int32 InNumOutputChannels, const float InSendLevel, EMixerSourceSubmixSendStage InSubmixSendStage, uint32 InSubmixId, FAlignedFloatBuffer& ScratchBuffer, FAlignedFloatBuffer& OutWetBuffer)
 	{
 		if (InSendLevel > 0.0f)
 		{
-			const FSourceInfo& SourceInfo = SourceInfos[SourceId];
+			FSourceInfo& SourceInfo = SourceInfos[SourceId];
 
 			// Don't need to mix into submixes if the source is paused
 			if (!SourceInfo.bIsPaused && !SourceInfo.bIsPausedForQuantization && !SourceInfo.bIsDone && SourceInfo.bIsPlaying)
 			{
-				const FMixerSourceSubmixOutputBuffer& SourceSubmixOutputBuffer = SourceSubmixOutputBuffers[SourceId];
-				SourceSubmixOutputBuffer.MixOutput(InSendLevel, InSubmixSendStage, OutWetBuffer);
+				FMixerSourceSubmixOutputBuffer& SourceSubmixOutputBuffer = SourceSubmixOutputBuffers[SourceId];
+
+				// Find the send matching this submix to check for per-send filter state.
+				// SubmixId is cached on the send to avoid atomic Pin() calls on every buffer tick.
+				FMixerSourceSubmixSend* MatchingSend = nullptr;
+				for (FMixerSourceSubmixSend& Send : SourceInfo.SubmixSends)
+				{
+					if (Send.SubmixId == InSubmixId)
+					{
+						MatchingSend = &Send;
+						break;
+					}
+				}
+
+				if (MatchingSend && MatchingSend->Filter.IsFiltering())
+				{
+					SourceSubmixOutputBuffer.MixOutput(InSendLevel, InSubmixSendStage, MatchingSend->Filter, MatchingSend->Lpf, MatchingSend->Hpf, ScratchBuffer, OutWetBuffer);
+				}
+				else
+				{
+					SourceSubmixOutputBuffer.MixOutput(InSendLevel, InSubmixSendStage, OutWetBuffer);
+				}
 			}
 		}
 	}
diff --git a/Engine/Source/Runtime/AudioMixer/Private/AudioMixerSourceManager.h b/Engine/Source/Runtime/AudioMixer/Private/AudioMixerSourceManager.h
--- a/Engine/Source/Runtime/AudioMixer/Private/AudioMixerSourceManager.h
+++ b/Engine/Source/Runtime/AudioMixer/Private/AudioMixerSourceManager.h
@@ -104,6 +104,9 @@ namespace Audio
 		// The submix ptr
 		FMixerSubmixWeakPtr Submix;
 
+		// Cached submix ID — avoids atomic Pin() calls in the render-thread hot path
+		uint32 SubmixId = TNumericLimits<uint32>::Max();
+
 		// The amount of audio that is to be mixed into this submix
 		float SendLevel = 0.0f;
 
@@ -116,6 +119,11 @@ namespace Audio
 		// If this is a soundfield submix, this is a pointer to the submix's Soundfield Factory.
 		// If this is nullptr, the submix is not a soundfield submix.
 		ISoundfieldFactory* SoundfieldFactory = nullptr;
+
+		FSendFilterSettings Filter;
+
+		FInterpolatedLPF Lpf;
+		FInterpolatedHPF Hpf;
 	};
 
 	// Struct holding mappings of bus ids (unique ids) to send level
@@ -124,6 +132,7 @@ namespace Audio
 		uint32 AudioBusId = INDEX_NONE;
 		float SendLevel = 0.0f;
 		int32 BusChannels = 0;
+		FSendFilterSettings Filter;
 #if ENABLE_AUDIO_DEBUG
 		FString AudioBusName;
 #endif // ENABLE_AUDIO_DEBUG
@@ -274,7 +283,7 @@ namespace Audio
 		void ComputeNextBlockOfSamples();
 		void UpdateSourceState();
 		void ClearStoppingSounds();
-		void MixOutputBuffers(const int32 SourceId, int32 InNumOutputChannels, const float InSendLevel, EMixerSourceSubmixSendStage InSubmixSendStage, FAlignedFloatBuffer& OutWetBuffer) const;
+		void MixOutputBuffers(const int32 SourceId, int32 InNumOutputChannels, const float InSendLevel, EMixerSourceSubmixSendStage InSubmixSendStage, uint32 InSubmixId, FAlignedFloatBuffer& ScratchBuffer, FAlignedFloatBuffer& OutWetBuffer);
 		void RenderSource(const int32 SourceId);
 
 		// Retrieves a channel map for the given source ID for the given output channels
@@ -292,7 +301,7 @@ namespace Audio
 		void SetSubmixSendInfo(const int32 SourceId, const FMixerSourceSubmixSend& SubmixSend);
 		void ClearSubmixSendInfo(const int32 SourceId, const FMixerSourceSubmixSend& SubmixSend);
 
-		void SetBusSendInfo(const int32 SourceId, EBusSendType InAudioBusSendType, uint32 AudiobusId, float BusSendLevel, const FString& InBusName);
+		void SetBusSendInfo(const int32 SourceId, EBusSendType InAudioBusSendType, uint32 AudiobusId, float BusSendLevel, const FString& InBusName, const FSendFilterSettings& InFilter = FSendFilterSettings{});
 
 		void UpdateDeviceChannelCount(const int32 InNumOutputChannels);
 
diff --git a/Engine/Source/Runtime/AudioMixer/Private/AudioMixerSourceOutputBuffer.cpp b/Engine/Source/Runtime/AudioMixer/Private/AudioMixerSourceOutputBuffer.cpp
--- a/Engine/Source/Runtime/AudioMixer/Private/AudioMixerSourceOutputBuffer.cpp
+++ b/Engine/Source/Runtime/AudioMixer/Private/AudioMixerSourceOutputBuffer.cpp
@@ -1,6 +1,8 @@
 // Copyright Epic Games, Inc. All Rights Reserved.
 
 #include "AudioMixerSourceOutputBuffer.h"
+#include "AudioMixerBus.h"
+#include "DSP/InterpolatedOnePole.h"
 #include "SoundFieldRendering.h"
 #include "AudioMixerSubmix.h"
 #include "DSP/FloatArrayMath.h"
@@ -500,6 +502,13 @@ namespace Audio
 		}
 	}
 
+	void FMixerSourceSubmixOutputBuffer::MixOutput(float InSendLevel, EMixerSourceSubmixSendStage InSubmixSendStage, const FSendFilterSettings& InFilter, FInterpolatedLPF& InLpf, FInterpolatedHPF& InHpf, FAlignedFloatBuffer& ScratchBuffer, FAlignedFloatBuffer& OutMixedBuffer)
+	{
+		const FAlignedFloatBuffer& SrcBuffer = (InSubmixSendStage == EMixerSourceSubmixSendStage::PostDistanceAttenuation) ? PostAttenuationOutputBuffer : PreAttenuationOutputBuffer;
+
+		Audio::ArrayFilterAndMixIn(SrcBuffer, OutMixedBuffer, InSendLevel, NumOutputChannels, InFilter.bEnableLPF, InFilter.LpfCutoff, InFilter.bEnableHPF, InFilter.HpfCutoff, InLpf, InHpf, ScratchBuffer);
+	}
+
 	FQuat FMixerSourceSubmixOutputBuffer::GetListenerRotation() const
 	{
 		return SoundfieldPositionalData.Rotation;
diff --git a/Engine/Source/Runtime/AudioMixer/Private/AudioMixerSourceOutputBuffer.h b/Engine/Source/Runtime/AudioMixer/Private/AudioMixerSourceOutputBuffer.h
--- a/Engine/Source/Runtime/AudioMixer/Private/AudioMixerSourceOutputBuffer.h
+++ b/Engine/Source/Runtime/AudioMixer/Private/AudioMixerSourceOutputBuffer.h
@@ -11,6 +11,10 @@
 
 namespace Audio
 {
+	class FInterpolatedLPF;
+	class FInterpolatedHPF;
+	struct FSendFilterSettings;
+
 	class FMixerSubmix;
 
 	struct FMixerSourceSubmixOutputBufferSettings
@@ -60,6 +64,7 @@ namespace Audio
 
 		// Called by submixes to mix this output buffer to their buffer
 		void MixOutput(float SendLevel, EMixerSourceSubmixSendStage InSubmixSendStage, FAlignedFloatBuffer& OutMixedBuffer) const;
+		void MixOutput(float InSendLevel, EMixerSourceSubmixSendStage InSubmixSendStage, const FSendFilterSettings& InFilter, FInterpolatedLPF& InLpf, FInterpolatedHPF& InHpf, FAlignedFloatBuffer& ScratchBuffer, FAlignedFloatBuffer& OutMixedBuffer);
 
 		// Return the listener rotation
 		// TODO: consolidate the code that is using this to be private to this class.
diff --git a/Engine/Source/Runtime/AudioMixer/Private/AudioMixerSourceVoice.cpp b/Engine/Source/Runtime/AudioMixer/Private/AudioMixerSourceVoice.cpp
--- a/Engine/Source/Runtime/AudioMixer/Private/AudioMixerSourceVoice.cpp
+++ b/Engine/Source/Runtime/AudioMixer/Private/AudioMixerSourceVoice.cpp
@@ -354,7 +354,7 @@ namespace Audio
 		return SourceManager->GetRelativeRenderCost(SourceId);
 	}
 
-	void FMixerSourceVoice::SetSubmixSendInfo(FMixerSubmixWeakPtr Submix, const float SendLevel, const EMixerSourceSubmixSendStage SendStage/* = EMixerSourceSubmixSendStage::PostDistanceAttenuation*/)
+	void FMixerSourceVoice::SetSubmixSendInfo(FMixerSubmixWeakPtr Submix, const float SendLevel, const EMixerSourceSubmixSendStage SendStage/* = EMixerSourceSubmixSendStage::PostDistanceAttenuation*/, const FSendFilterSettings& InFilter/* = FSendFilterSettings{}*/)
 	{
 		AUDIO_MIXER_CHECK_GAME_THREAD(MixerDevice);
 
@@ -370,14 +370,18 @@ namespace Audio
 				NewSubmixSend.SendLevel = SendLevel;
 				NewSubmixSend.bIsMainSend = false;
 				NewSubmixSend.SubmixSendStage = SendStage;
+				NewSubmixSend.Filter = InFilter;
 
 				SubmixSends.Add(SubmixPtr->GetId(), NewSubmixSend);
 				SourceManager->SetSubmixSendInfo(SourceId, NewSubmixSend);
 			}
-			else if (!FMath::IsNearlyEqual(SubmixSend->SendLevel, SendLevel, SourceManager->GetFloatCompareTolerance()) || SubmixSend->SubmixSendStage != SendStage)
+			else if (!FMath::IsNearlyEqual(SubmixSend->SendLevel, SendLevel, SourceManager->GetFloatCompareTolerance())
+				|| SubmixSend->SubmixSendStage != SendStage
+				|| SubmixSend->Filter.HasChanged(InFilter))
 			{
 				SubmixSend->SendLevel = SendLevel;
 				SubmixSend->SubmixSendStage = SendStage;
+				SubmixSend->Filter = InFilter;
 				SourceManager->SetSubmixSendInfo(SourceId, *SubmixSend);
 			}
 		}
@@ -418,7 +422,7 @@ namespace Audio
 	}
 
 
-	void FMixerSourceVoice::SetAudioBusSendInfo(EBusSendType InBusSendType, uint32 AudioBusId, float BusSendLevel, const FString& InBusName)
+	void FMixerSourceVoice::SetAudioBusSendInfo(EBusSendType InBusSendType, uint32 AudioBusId, float BusSendLevel, const FString& InBusName, const FSendFilterSettings& InFilter)
 	{
 		AUDIO_MIXER_CHECK_GAME_THREAD(MixerDevice);
 
@@ -427,7 +431,7 @@ namespace Audio
 			BusSendLevel = 0.0f;
 		}
 
-		SourceManager->SetBusSendInfo(SourceId, InBusSendType, AudioBusId, BusSendLevel, InBusName);
+		SourceManager->SetBusSendInfo(SourceId, InBusSendType, AudioBusId, BusSendLevel, InBusName, InFilter);
 	}
 
 	bool FMixerSourceVoice::IsRenderingToSubmixes() const
diff --git a/Engine/Source/Runtime/AudioMixer/Private/AudioMixerSourceVoice.h b/Engine/Source/Runtime/AudioMixer/Private/AudioMixerSourceVoice.h
--- a/Engine/Source/Runtime/AudioMixer/Private/AudioMixerSourceVoice.h
+++ b/Engine/Source/Runtime/AudioMixer/Private/AudioMixerSourceVoice.h
@@ -128,8 +128,8 @@ namespace Audio
 		// Useful as a metric more human-centric than CPU utilization.
 		float GetRelativeRenderCost() const;
 
-		// Sets the submix send levels
-		void SetSubmixSendInfo(FMixerSubmixWeakPtr Submix, const float SendLevel, const EMixerSourceSubmixSendStage SendStage = EMixerSourceSubmixSendStage::PostDistanceAttenuation);
+		// Sets the send level and optional filter for a single submix send
+		void SetSubmixSendInfo(FMixerSubmixWeakPtr Submix, const float SendLevel, const EMixerSourceSubmixSendStage SendStage = EMixerSourceSubmixSendStage::PostDistanceAttenuation, const FSendFilterSettings& InFilter = FSendFilterSettings{});
 
 		// Clears the submix send to the given submix
 		void ClearSubmixSendInfo(FMixerSubmixWeakPtr Submix);
@@ -141,7 +141,7 @@ namespace Audio
 		void SetEnablement(bool bInEnableBusSendRouting, bool bInEnableMainSubmixOutput, bool bInEnableSubmixSendRouting);
 
 		// Set the source bus send levels
-		void SetAudioBusSendInfo(EBusSendType InBusSendType, uint32 AudioBusId, float BusSendLevel, const FString& InBusName);
+		void SetAudioBusSendInfo(EBusSendType InBusSendType, uint32 AudioBusId, float BusSendLevel, const FString& InBusName, const FSendFilterSettings& InFilter = FSendFilterSettings{});
 
 	private:
 
diff --git a/Engine/Source/Runtime/AudioMixer/Private/AudioMixerSubmix.cpp b/Engine/Source/Runtime/AudioMixer/Private/AudioMixerSubmix.cpp
--- a/Engine/Source/Runtime/AudioMixer/Private/AudioMixerSubmix.cpp
+++ b/Engine/Source/Runtime/AudioMixer/Private/AudioMixerSubmix.cpp
@@ -1453,6 +1453,11 @@ namespace Audio
 			CSV_SCOPED_TIMING_STAT(Audio, SubmixSource);
 			CONDITIONAL_SCOPE_CYCLE_COUNTER(STAT_AudioMixerSubmixSource, (MixerSourceIds.Num() > 0));
 
+			// Pre-size the scratch buffer to match InputBuffer so it never needs to grow mid-loop.
+			// FilterScratchBuffer is a member, so its backing allocation persists across MixBuffer()
+			// calls; SetNumUninitialized only reallocates if the current capacity is insufficient.
+			FilterScratchBuffer.SetNumUninitialized(InputBuffer.Num());
+
 			// Loop through this submix's sound sources
 			for (const auto& MixerSourceVoiceIter : MixerSourceIds)
 			{
@@ -1460,8 +1465,8 @@ namespace Audio
 				const float SendLevel = MixerSourceVoiceIter.Value.SendLevel;
 				const EMixerSourceSubmixSendStage SubmixSendStage = MixerSourceVoiceIter.Value.SubmixSendStage;
 
-				MixerDevice->GetSourceManager()->MixOutputBuffers(SourceId, NumChannels, SendLevel, SubmixSendStage, InputBuffer);
-				
+				MixerDevice->GetSourceManager()->MixOutputBuffers(SourceId, NumChannels, SendLevel, SubmixSendStage, GetId(), FilterScratchBuffer, InputBuffer);
+
 				// Check the buffer after each voice mix to catch any bad values.
 				AUDIO_CHECK_BUFFER_NAMED_MSG(InputBuffer, TEXT("Submix SourceMix"), TEXT("Submix: %s"), *GetName());
 			}
diff --git a/Engine/Source/Runtime/AudioMixer/Public/AudioMixerBus.h b/Engine/Source/Runtime/AudioMixer/Public/AudioMixerBus.h
--- a/Engine/Source/Runtime/AudioMixer/Public/AudioMixerBus.h
+++ b/Engine/Source/Runtime/AudioMixer/Public/AudioMixerBus.h
@@ -5,6 +5,7 @@
 #include "Audio.h"
 #include "AudioBusSubsystem.h"
 #include "AudioMixerTrace.h"
+#include "DSP/InterpolatedOnePole.h"
 #include "DSP/MultithreadedPatching.h"
 #include "IAudioMixerRenderStep.h"
 #include "Sound/AudioSettings.h"
@@ -17,12 +18,46 @@ namespace Audio
 {
 	class FMixerSourceManager;
 
+	// Consolidated filter settings for a single bus/submix send.
+	// Passed between the game thread and render thread as a simple value type.
+	struct FSendFilterSettings
+	{
+		uint8 bEnableLPF : 1 = 0;
+		uint8 bEnableHPF : 1 = 0;
+		float LpfCutoff = 20000.0f;
+		float HpfCutoff = 20.0f;
+
+		FSendFilterSettings() = default;
+		FSendFilterSettings(bool bInEnableLPF, bool bInEnableHPF, float InLpfCutoff, float InHpfCutoff)
+			: LpfCutoff(InLpfCutoff), HpfCutoff(InHpfCutoff)
+		{
+			bEnableLPF = bInEnableLPF;
+			bEnableHPF = bInEnableHPF;
+		}
+
+		bool IsFiltering() const { return bEnableLPF || bEnableHPF; }
+
+		// Returns true if any filter parameter differs significantly from Other.
+		// Use to determine whether a render-thread update is needed.
+		bool HasChanged(const FSendFilterSettings& Other) const
+		{
+			return bEnableLPF != Other.bEnableLPF
+				|| bEnableHPF != Other.bEnableHPF
+				|| (bEnableLPF && !FMath::IsNearlyEqual(LpfCutoff, Other.LpfCutoff))
+				|| (bEnableHPF && !FMath::IsNearlyEqual(HpfCutoff, Other.HpfCutoff));
+		}
+	};
+
 	// Struct holding mappings of runtime source ids (bus instances) to bus send level
 	struct FAudioBusSend
 	{
 		int32 SourceId = INDEX_NONE;
 		uint32 WaveInstancePlayOrder = INDEX_NONE;
 		float SendLevel = 0.0f;
+
+		FSendFilterSettings Filter;
+		FInterpolatedLPF Lpf;
+		FInterpolatedHPF Hpf;
 	};
 
 	// Bus instance data. Holds source id bus instances and bus sends data
@@ -112,6 +147,10 @@ namespace Audio
 		// Buses feed audio to each other by storing their previous buffer. Current buses mix in previous other buses (including themselves)
 		FAlignedFloatBuffer MixedSourceData[2];
 
+		// Scratch buffer used for intermediate filter output in MixBuffer(). Kept as a member so its
+		// heap allocation persists across render callbacks and avoids per-frame reallocation.
+		FAlignedFloatBuffer FilterScratchBuffer;
+
 		// The index of the bus data currently being rendered
 		int32 CurrentBufferIndex;
 
diff --git a/Engine/Source/Runtime/AudioMixer/Public/AudioMixerSubmix.h b/Engine/Source/Runtime/AudioMixer/Public/AudioMixerSubmix.h
--- a/Engine/Source/Runtime/AudioMixer/Public/AudioMixerSubmix.h
+++ b/Engine/Source/Runtime/AudioMixer/Public/AudioMixerSubmix.h
@@ -437,6 +437,7 @@ namespace Audio
 
 		FAlignedFloatBuffer ScratchBuffer;
 		FAlignedFloatBuffer SubmixChainMixBuffer;
+		FAlignedFloatBuffer FilterScratchBuffer;
 		FAlignedFloatBuffer InputBuffer;
 		FAlignedFloatBuffer DownmixedBuffer;
 		FAlignedFloatBuffer SourceInputBuffer;
diff --git a/Engine/Source/Runtime/Engine/Classes/Sound/SoundSourceBusSend.h b/Engine/Source/Runtime/Engine/Classes/Sound/SoundSourceBusSend.h
--- a/Engine/Source/Runtime/Engine/Classes/Sound/SoundSourceBusSend.h
+++ b/Engine/Source/Runtime/Engine/Classes/Sound/SoundSourceBusSend.h
@@ -4,6 +4,7 @@
 
 #include "CoreMinimal.h"
 #include "Curves/CurveFloat.h"
+#include "Sound/SoundSubmixSend.h"
 #include "SoundSourceBusSend.generated.h"
 
 class USoundSourceBus;
@@ -67,6 +68,12 @@ struct FSoundSourceBusSendInfo
 	UPROPERTY(EditAnywhere, BlueprintReadWrite, Category = BusSend, meta = (EditCondition = "SourceBusSendLevelControlMethod == ESourceBusSendLevelControlMethod::CustomCurve", EditConditionHides))
 	FRuntimeFloatCurve CustomSendLevelCurve;
 
+	UPROPERTY(EditAnywhere, BlueprintReadWrite, Category="HPF/LPF", meta=(ClampMin="20.0", ClampMax="20000.0", UIMin="20.0", UIMax="20000.0"))
+	TOptional<float> LPFCutoffFrequency;
+
+	UPROPERTY(EditAnywhere, BlueprintReadWrite, Category="HPF/LPF", meta=(ClampMin="20.0", ClampMax="20000.0", UIMin="20.0", UIMax="20000.0"))
+	TOptional<float> HPFCutoffFrequency;
+
 	FSoundSourceBusSendInfo()
 		: SourceBusSendLevelControlMethod(ESourceBusSendLevelControlMethod::Manual)
 		, SoundSourceBus(nullptr)
diff --git a/Engine/Source/Runtime/Engine/Classes/Sound/SoundSubmixSend.h b/Engine/Source/Runtime/Engine/Classes/Sound/SoundSubmixSend.h
--- a/Engine/Source/Runtime/Engine/Classes/Sound/SoundSubmixSend.h
+++ b/Engine/Source/Runtime/Engine/Classes/Sound/SoundSubmixSend.h
@@ -86,7 +86,7 @@ struct FSoundSubmixSendInfoBase
 
 	ENGINE_API FSoundSubmixSendInfoBase();
 
-	// // Changes the method we use to determine the send level - see enum value tooltips for more information
+	// Changes the method we use to determine the send level - see enum value tooltips for more information
 	UPROPERTY(EditAnywhere, BlueprintReadWrite, Category = SubmixSend)
 	ESendLevelControlMethod SendLevelControlMethod;
 	
@@ -125,6 +125,12 @@ struct FSoundSubmixSendInfoBase
 	// The custom send curve to use for distance-based send level. (0.0-1.0 on the curve's X-axis maps to Min/Max Send Distance)
 	UPROPERTY(EditAnywhere, BlueprintReadWrite, Category = SubmixSend, meta = (EditCondition = "SendLevelControlMethod == ESendLevelControlMethod::CustomCurve", EditConditionHides))
 	FRuntimeFloatCurve CustomSendLevelCurve;
+
+	UPROPERTY(EditAnywhere, BlueprintReadWrite, Category="HPF/LPF", meta=(ClampMin="20.0", ClampMax="20000.0", UIMin="20.0", UIMax="20000.0"))
+	TOptional<float> LPFCutoffFrequency;
+
+	UPROPERTY(EditAnywhere, BlueprintReadWrite, Category="HPF/LPF", meta=(ClampMin="20.0", ClampMax="20000.0", UIMin="20.0", UIMax="20000.0"))
+	TOptional<float> HPFCutoffFrequency;
 };
 
 UENUM(BlueprintType)
diff --git a/Engine/Source/Runtime/Engine/Private/ActiveSound.cpp b/Engine/Source/Runtime/Engine/Private/ActiveSound.cpp
--- a/Engine/Source/Runtime/Engine/Private/ActiveSound.cpp
+++ b/Engine/Source/Runtime/Engine/Private/ActiveSound.cpp
@@ -588,9 +588,12 @@ void FActiveSound::SetSourceBusSend(EBusSendType BusSendType, const FSoundSource
 	// Override send level if the source bus send is already included in active sound
 	for (FSoundSourceBusSendInfo& Info : BusSendsOverride[(int32)BusSendType])
 	{
-		if (Info.SoundSourceBus == SendInfo.SoundSourceBus || Info.AudioBus == SendInfo.AudioBus)
+		if ((SendInfo.SoundSourceBus && Info.SoundSourceBus == SendInfo.SoundSourceBus) ||
+			(SendInfo.AudioBus && Info.AudioBus == SendInfo.AudioBus))
 		{
 			Info.SendLevel = SendInfo.SendLevel;
+			Info.LPFCutoffFrequency = SendInfo.LPFCutoffFrequency;
+			Info.HPFCutoffFrequency = SendInfo.HPFCutoffFrequency;
 
 			bHasNewBusSends = true;
 			NewBusSends.Add(TTuple<EBusSendType, FSoundSourceBusSendInfo>(BusSendType, SendInfo));
@@ -763,6 +766,8 @@ void FActiveSound::GetBusSends(EBusSendType BusSendType, TArray<FSoundSourceBusS
 				if (bSameSourceBus || bSameAudioBus)
 				{
 					OutSendInfo.SendLevel = SendInfo.SendLevel;
+					OutSendInfo.LPFCutoffFrequency = SendInfo.LPFCutoffFrequency;
+					OutSendInfo.HPFCutoffFrequency = SendInfo.HPFCutoffFrequency;
 					bOverridden = true;
 					break;
 				}
diff --git a/Engine/Source/Runtime/Engine/Private/Components/AudioComponent.cpp b/Engine/Source/Runtime/Engine/Private/Components/AudioComponent.cpp
--- a/Engine/Source/Runtime/Engine/Private/Components/AudioComponent.cpp
+++ b/Engine/Source/Runtime/Engine/Private/Components/AudioComponent.cpp
@@ -1611,7 +1611,23 @@ void UAudioComponent::SetBusSendEffectInternal(USoundSourceBus* InSourceBus, UAu
 			DECLARE_CYCLE_STAT(TEXT("FAudioThreadTask.AudioSetBusSend"), STAT_SetBusSend, STATGROUP_AudioThreadCommands);
 			AudioDevice->SendCommandToActiveSounds(AudioComponentID, [InBusSendType, BusSendInfo](FActiveSound& ActiveSound)
 			{
-				ActiveSound.SetSourceBusSend(InBusSendType, BusSendInfo);
+				// Preserve existing filter state - this function only changes send level,
+				// so don't clobber filter enable/cutoff fields on an existing override entry.
+				FSoundSourceBusSendInfo InfoToSend = BusSendInfo;
+				TArray<FSoundSourceBusSendInfo> CurrentSends;
+				ActiveSound.GetBusSends(InBusSendType, CurrentSends);
+				for (const FSoundSourceBusSendInfo& Existing : CurrentSends)
+				{
+					const bool bSameSourceBus = InfoToSend.SoundSourceBus && (InfoToSend.SoundSourceBus == Existing.SoundSourceBus);
+					const bool bSameAudioBus  = InfoToSend.AudioBus && (InfoToSend.AudioBus == Existing.AudioBus);
+					if (bSameSourceBus || bSameAudioBus)
+					{
+						InfoToSend.LPFCutoffFrequency = Existing.LPFCutoffFrequency;
+						InfoToSend.HPFCutoffFrequency = Existing.HPFCutoffFrequency;
+						break;
+					}
+				}
+				ActiveSound.SetSourceBusSend(InBusSendType, InfoToSend);
 			}, GET_STATID(STAT_SetBusSend));
 		}
 	}
diff --git a/Engine/Source/Runtime/SignalProcessing/Private/FloatArrayMath.cpp b/Engine/Source/Runtime/SignalProcessing/Private/FloatArrayMath.cpp
--- a/Engine/Source/Runtime/SignalProcessing/Private/FloatArrayMath.cpp
+++ b/Engine/Source/Runtime/SignalProcessing/Private/FloatArrayMath.cpp
@@ -5,6 +5,7 @@
 #include "SignalProcessingModule.h"
 #include "ProfilingDebugging/CsvProfiler.h"
 #include "DSP/Dsp.h"
+#include "DSP/InterpolatedOnePole.h"
 
 #if INTEL_ISPC && !UE_BUILD_SHIPPING
 #include "HAL/IConsoleManager.h"
@@ -2415,6 +2416,58 @@ namespace Audio
 		}
 	}
 
+	void ArrayFilterAndMixIn(TArrayView<const float> InFloatBuffer, TArrayView<float> BufferToSumTo, const float Gain, const int32 InNumChannels, bool bEnableLPF, float InLpfCutoff, bool bEnableHPF, float InHpfCutoff, FInterpolatedLPF& InLpf, FInterpolatedHPF& InHpf, FAlignedFloatBuffer& ScratchBuffer)
+	{
+		CSV_SCOPED_TIMING_STAT(Audio_Dsp, ArrayFilterAndMixIn);
+
+		if (!ensureMsgf(InFloatBuffer.Num() == BufferToSumTo.Num(), TEXT("ArrayFilterAndMixIn: input (%d) and output (%d) buffers must be equal size"), InFloatBuffer.Num(), BufferToSumTo.Num()))
+		{
+			return;
+		}
+		
+		if (!ensureMsgf(InNumChannels > 0, TEXT("ArrayFilterAndMixIn: InNumChannels must be > 0")))
+		{
+			return;
+		}
+		checkf(InNumChannels <= 1 || InFloatBuffer.Num() % InNumChannels == 0, TEXT("NumSamples must be frame-aligned (multiple of InNumChannels)"));
+
+		const int32 NumSamples = InFloatBuffer.Num();
+
+		const bool bBypassLPF = !bEnableLPF;
+		const bool bBypassHPF = !bEnableHPF;
+
+		if (bBypassLPF && bBypassHPF)
+		{
+			ArrayMixIn(InFloatBuffer, BufferToSumTo, Gain);
+			return;
+		}
+
+		// Copy source into scratch and filter in place.
+		ScratchBuffer.SetNumUninitialized(NumSamples);
+		FMemory::Memcpy(ScratchBuffer.GetData(), InFloatBuffer.GetData(), NumSamples * sizeof(float));
+
+		// Coefficient steps once per frame across all SIMD paths, so InterpLength must be NumFrames not NumSamples.
+		// StartFrequencyInterpolation + Process + StopFrequencyInterpolation must always be called as a matched set -
+		// failing to Stop after exactly InterpLength frames causes the coefficient to drift and the filter to become unstable.
+		const int32 NumFrames = NumSamples / FMath::Max(InNumChannels, 1);
+
+		if (!bBypassLPF)
+		{
+			InLpf.StartFrequencyInterpolation(InLpfCutoff, NumFrames);
+			InLpf.ProcessBufferInPlace(ScratchBuffer.GetData(), NumSamples);
+			InLpf.StopFrequencyInterpolation();
+		}
+
+		if (!bBypassHPF)
+		{
+			InHpf.StartFrequencyInterpolation(InHpfCutoff, NumFrames);
+			InHpf.ProcessBufferInPlace(ScratchBuffer.GetData(), NumSamples);
+			InHpf.StopFrequencyInterpolation();
+		}
+
+		ArrayMixIn(ScratchBuffer, BufferToSumTo, Gain);
+	}
+
 	void ArrayFloatToPcm16(TArrayView<const float> InView, TArrayView<int16> OutView)
 	{
 		CSV_SCOPED_TIMING_STAT(Audio_Dsp, ArrayFloatToPcm16);
diff --git a/Engine/Source/Runtime/SignalProcessing/Private/InterpolatedOnePole.cpp b/Engine/Source/Runtime/SignalProcessing/Private/InterpolatedOnePole.cpp
--- a/Engine/Source/Runtime/SignalProcessing/Private/InterpolatedOnePole.cpp
+++ b/Engine/Source/Runtime/SignalProcessing/Private/InterpolatedOnePole.cpp
@@ -17,11 +17,11 @@ namespace Audio
 
 	void FInterpolatedLPF::Init(float InSampleRate, int32 InNumChannels)
 	{
+		checkf(InNumChannels > 0, TEXT("FInterpolatedLPF::Init called with InNumChannels <= 0"));
 		SampleRate = InSampleRate;
 		NumChannels = InNumChannels;
 		CutoffFrequency = -1.0f;
 		Z1.Init(0.0f, NumChannels);
-		Z1Data = Z1.GetData();
 		Reset();
 	}
 
@@ -41,7 +41,7 @@ namespace Audio
 
 			const float NormalizedFreq = FMath::Clamp(2.0f * GetCutoffFrequency() / SampleRate, 0.0f, 1.0f);
 			B1Target = FMath::Exp(-PI * NormalizedFreq);
-			B1Delta = (B1Target - B1Curr) / static_cast<float>(CurrInterpLength);
+			B1Delta = (CurrInterpLength > 1) ? (B1Target - B1Curr) / static_cast<float>(CurrInterpLength) : 0.0f;
 		}
 
 		if (CurrInterpLength <= 1)
@@ -52,6 +52,7 @@ namespace Audio
 
 	void FInterpolatedLPF::ProcessAudioFrame(const float* RESTRICT InputFrame, float* RESTRICT OutputFrame)
 	{
+		float* Z1Data = Z1.GetData();
 		B1Curr += B1Delta; // step forward coefficient
 
 		/*
@@ -77,6 +78,7 @@ namespace Audio
 
 	void FInterpolatedLPF::ProcessAudioBuffer(const float *RESTRICT InputBuffer, float *RESTRICT OutputBuffer, const int32 NumSamples)
 	{
+		float* Z1Data = Z1.GetData();
 		if (NumChannels == 1)
 		{
 			// keep these in registers instead of having a load\store on every loop iteration
@@ -89,7 +91,7 @@ namespace Audio
 				B1 += Delta;
 
 				const float InputSample = InputBuffer[SampleIndex];
-				float Yn = InputSample + B1Curr * (Z1D - InputSample); // LPF
+				float Yn = InputSample + B1 * (Z1D - InputSample); // LPF
 				Z1D = Yn;
 				OutputBuffer[SampleIndex] = Yn;
 			}
@@ -175,6 +177,7 @@ namespace Audio
 
 	void FInterpolatedLPF::ProcessBufferInPlace(float* InOutBuffer, const int32 NumSamples)
 	{
+		float* Z1Data = Z1.GetData();
 		if (NumChannels == 1)
 		{
 			// keep these in registers instead of having a load\store on every loop iteration
@@ -187,7 +190,7 @@ namespace Audio
 				B1 += Delta;
 
 				const float InputSample = InOutBuffer[SampleIndex];
-				float Yn = InputSample + B1Curr * (Z1D - InputSample); // LPF
+				float Yn = InputSample + B1 * (Z1D - InputSample); // LPF
 				Z1D = Yn;
 				InOutBuffer[SampleIndex] = Yn;
 			}
@@ -278,7 +281,6 @@ namespace Audio
 		B1Target = B1Curr;
 		CurrInterpLength = 0;
 		ClearMemory();
-		Z1Data = Z1.GetData();
 		bIsFirstFrequencyChange = true;
 	}
 
@@ -317,12 +319,12 @@ namespace Audio
 
 	void FInterpolatedHPF::Init(float InSampleRate, int32 InNumChannels)
 	{
+		checkf(InNumChannels > 0, TEXT("FInterpolatedHPF::Init called with InNumChannels <= 0"));
 		SampleRate = InSampleRate;
 		NyquistLimit = 0.5f * SampleRate - 1.0f;
 		NumChannels = InNumChannels;
 		CutoffFrequency = -1.0f;
 		Z1.Init(0.0f, NumChannels);
-		Z1Data = Z1.GetData();
 		Reset();
 	}
 
@@ -347,7 +349,7 @@ namespace Audio
 			const float G = Audio::FastTan(PI * GetCutoffFrequency() / SampleRate);
 
 			A0Target = G / (1.0f + G);
-			A0Delta = (A0Target - A0Curr) / static_cast<float>(CurrInterpLength);
+			A0Delta = (CurrInterpLength > 1) ? (A0Target - A0Curr) / static_cast<float>(CurrInterpLength) : 0.0f;
 		}
 
 		if (CurrInterpLength <= 1)
@@ -358,6 +360,7 @@ namespace Audio
 
 	void FInterpolatedHPF::ProcessAudioFrame(const float* RESTRICT InputFrame, float* RESTRICT OutputFrame)
 	{
+		float* Z1Data = Z1.GetData();
 		A0Curr += A0Delta; // step forward coefficient
 
 		for (int32 i = 0; i < NumChannels; ++i)
@@ -373,6 +376,7 @@ namespace Audio
 
 	void FInterpolatedHPF::ProcessAudioBuffer(const float *RESTRICT InputBuffer, float *RESTRICT OutputBuffer, const int32 NumSamples)
 	{
+		float* Z1Data = Z1.GetData();
 		if (NumChannels == 1)
 		{
 			// keep these in registers instead of having a load\store on every loop iteration
@@ -480,13 +484,121 @@ namespace Audio
 		}
 	}
 
+	void FInterpolatedHPF::ProcessBufferInPlace(float* InOutBuffer, const int32 NumSamples)
+	{
+		float* Z1Data = Z1.GetData();
+		if (NumSamples <= 0 || NumChannels <= 0)
+		{
+			return;
+		}
+
+		if (NumChannels == 1)
+		{
+			float A0 = A0Curr;
+			float Z1D = Z1Data[0];
+			float Delta = A0Delta;
+			for (int32 SampleIndex = 0; SampleIndex < NumSamples; ++SampleIndex)
+			{
+				A0 += Delta;
+
+				const float InputSample = InOutBuffer[SampleIndex];
+				const float Vn = (InputSample - Z1D) * A0;
+				const float LPF = Vn + Z1D;
+				Z1D = Vn + LPF;
+
+				InOutBuffer[SampleIndex] = InputSample - LPF;
+			}
+
+			Z1Data[0] = Z1D;
+			A0Curr = A0;
+
+			return;
+		}
+
+		int32 SampleIndex = 0;
+		const int32 NumToSimd = NumSamples & Simd4Mask;
+		if (NumChannels == 2)
+		{
+			VectorRegister4Float A0CurrVec = VectorLoadFloat1(&A0Curr);
+			VectorRegister4Float A0DeltaVec = MakeVectorRegister(A0Delta + A0Delta, A0Delta + A0Delta, A0Delta + A0Delta, A0Delta + A0Delta);
+			VectorRegister4Float Z1DataVec = MakeVectorRegister(Z1Data[0], Z1Data[1], Z1Data[0], Z1Data[1]);
+
+			if (NumToSimd)
+			{
+				A0CurrVec = VectorAdd(A0CurrVec, MakeVectorRegister(A0Delta, A0Delta, A0Delta + A0Delta, A0Delta + A0Delta));
+				}
+
+			for (; SampleIndex < NumToSimd; SampleIndex += AUDIO_NUM_FLOATS_PER_VECTOR_REGISTER)
+			{
+				VectorRegister4Float InputSample = VectorLoad(&InOutBuffer[SampleIndex]);
+				VectorRegister4Float Vn = VectorMultiply(VectorSubtract(InputSample, Z1DataVec), A0CurrVec);
+				VectorRegister4Float LPF = VectorAdd(Vn, Z1DataVec);
+
+				VectorRegister4Float VnLPFSum = VectorAdd(Vn, LPF);
+				Z1DataVec = VectorShuffle(Z1DataVec, VnLPFSum, 0, 1, 0, 1);
+				Vn = VectorMultiply(VectorSubtract(InputSample, Z1DataVec), A0CurrVec);
+				LPF = VectorAdd(Vn, Z1DataVec);
+
+				VnLPFSum = VectorAdd(Vn, LPF);
+				Z1DataVec = VectorSwizzle(VnLPFSum, 2, 3, 2, 3);
+
+				VectorStore(VectorSubtract(InputSample, LPF), &InOutBuffer[SampleIndex]);
+
+				A0CurrVec = VectorAdd(A0CurrVec, A0DeltaVec);
+			}
+
+			float Store[4];
+			VectorStore(A0CurrVec, Store);
+			A0Curr = Store[0];
+			VectorStore(Z1DataVec, Store);
+			Z1Data[0] = Store[0];
+			Z1Data[1] = Store[1];
+		}
+		else if (NumChannels == 4)
+		{
+			VectorRegister4Float A0CurrVec = VectorLoadFloat1(&A0Curr);
+			VectorRegister4Float A0DeltaVec = VectorLoadFloat1(&A0Delta);
+			VectorRegister4Float Z1DataVec = VectorLoad(&Z1Data[0]);
+
+			for (; SampleIndex < NumToSimd; SampleIndex += AUDIO_NUM_FLOATS_PER_VECTOR_REGISTER)
+			{
+				A0CurrVec = VectorAdd(A0CurrVec, A0DeltaVec);
+
+				VectorRegister4Float InputSample = VectorLoad(&InOutBuffer[SampleIndex]);
+				VectorRegister4Float Vn = VectorMultiply(VectorSubtract(InputSample, Z1DataVec), A0CurrVec);
+				VectorRegister4Float LPF = VectorAdd(Vn, Z1DataVec);
+
+				Z1DataVec = VectorAdd(Vn, LPF);
+
+				VectorStore(VectorSubtract(InputSample, LPF), &InOutBuffer[SampleIndex]);
+			}
+
+			float Store[4];
+			VectorStore(A0CurrVec, Store);
+			A0Curr = Store[0];
+			VectorStore(Z1DataVec, &Z1Data[0]);
+		}
+
+		for (; SampleIndex < NumSamples; ++SampleIndex)
+		{
+			const int32 ChannelIndex = SampleIndex % NumChannels;
+			A0Curr += A0Delta * !ChannelIndex;
+
+			const float InputSample = InOutBuffer[SampleIndex];
+			const float Vn = (InputSample - Z1Data[ChannelIndex]) * A0Curr;
+			const float LPF = Vn + Z1Data[ChannelIndex];
+			Z1Data[ChannelIndex] = Vn + LPF;
+
+			InOutBuffer[SampleIndex] = InputSample - LPF;
+		}
+	}
+
 	void FInterpolatedHPF::Reset()
 	{
 		A0Curr = 0.0f;
 		A0Delta = 0.0f;
 		CurrInterpLength = 0;
 		ClearMemory();
-		Z1Data = Z1.GetData();
 		bIsFirstFrequencyChange = true;
 	}
 
diff --git a/Engine/Source/Runtime/SignalProcessing/Public/DSP/FloatArrayMath.h b/Engine/Source/Runtime/SignalProcessing/Public/DSP/FloatArrayMath.h
--- a/Engine/Source/Runtime/SignalProcessing/Public/DSP/FloatArrayMath.h
+++ b/Engine/Source/Runtime/SignalProcessing/Public/DSP/FloatArrayMath.h
@@ -243,6 +243,10 @@ namespace Audio
 	/** This version of ArrayMixIn will fade from StartGain to EndGain. */
 	SIGNALPROCESSING_API void ArrayMixIn(TArrayView<const float> InFloatBuffer, TArrayView<float> BufferToSumTo, const float StartGain, const float EndGain);
 
+	class FInterpolatedLPF;
+	class FInterpolatedHPF;
+	SIGNALPROCESSING_API void ArrayFilterAndMixIn(TArrayView<const float> InFloatBuffer, TArrayView<float> BufferToSumTo, const float Gain, const int32 InNumChannels, bool bEnableLPF, float InLpfCutoff, bool bEnableHPF, float InHpfCutoff, FInterpolatedLPF& InLpf, FInterpolatedHPF& InHpf, FAlignedFloatBuffer& ScratchBuffer);
+
 	SIGNALPROCESSING_API void ArrayFloatToPcm16(TArrayView<const float> InView, TArrayView<int16> OutView);
 	SIGNALPROCESSING_API void ArrayPcm16ToFloat(TArrayView<const int16> InView, TArrayView<float> OutView);
 
diff --git a/Engine/Source/Runtime/SignalProcessing/Public/DSP/InterpolatedOnePole.h b/Engine/Source/Runtime/SignalProcessing/Public/DSP/InterpolatedOnePole.h
--- a/Engine/Source/Runtime/SignalProcessing/Public/DSP/InterpolatedOnePole.h
+++ b/Engine/Source/Runtime/SignalProcessing/Public/DSP/InterpolatedOnePole.h
@@ -67,7 +67,6 @@ namespace Audio
 		float B1Delta{ 0.0f }; // coefficient step size 
 		float B1Target{ 0.0f };
 		TArray<float> Z1; // multi-channel delay terms
-		float* Z1Data{ nullptr };
 		int32 CurrInterpLength{ 0 };
 		int32 NumInterpSteps;
 		float SampleRate{ 0 };
@@ -108,6 +107,9 @@ namespace Audio
 		// interpolates coefficient and processes a buffer
 		SIGNALPROCESSING_API void ProcessAudioBuffer(const float* RESTRICT InputBuffer, float* RESTRICT OutputBuffer, const int32 NumSamples);
 
+		// interpolates coefficients and processes a buffer in place
+		SIGNALPROCESSING_API void ProcessBufferInPlace(float* InOutBuffer, int32 NumSamples);
+
 		/*
 			StopFrequencyInterpolation() needs to be called manually when the interpolation should be done.
 			Snaps the coefficient to the target value.
@@ -135,7 +137,6 @@ namespace Audio
 		float A0Delta{ 0.0f }; // coefficient step size
 		float A0Target{ 0.0f };
 		TArray<float> Z1; // multi-channel delay terms
-		float* Z1Data{ nullptr };
 		int32 CurrInterpLength{ 0 };
 		int32 NumInterpSteps;
 		float SampleRate{ 0 };
diff --git a/Engine/Source/Runtime/SignalProcessing/Tests/InterpolatedOnePoleTests.cpp b/Engine/Source/Runtime/SignalProcessing/Tests/InterpolatedOnePoleTests.cpp
new file mode 100644
--- /dev/null
+++ b/Engine/Source/Runtime/SignalProcessing/Tests/InterpolatedOnePoleTests.cpp
@@ -0,0 +1,321 @@
+// Copyright Epic Games, Inc. All Rights Reserved.
+
+#include "DSP/InterpolatedOnePole.h"
+#include "Misc/AutomationTest.h"
+
+#if WITH_DEV_AUTOMATION_TESTS
+
+namespace Audio::Tests
+{
+	static constexpr float TestSampleRate = 48000.0f;
+	static constexpr float TestCutoffHz   = 1000.0f;
+
+	// ------------------------------------------------------------------ helpers
+
+	// Run a filter for NumSamples frames of mono constant-value input.
+	// Returns the last output sample.
+	static float RunLPFConstantMono(FInterpolatedLPF& Filter, float InputValue, int32 NumSamples)
+	{
+		float Output = 0.0f;
+		for (int32 i = 0; i < NumSamples; ++i)
+		{
+			Filter.ProcessAudioFrame(&InputValue, &Output);
+		}
+		return Output;
+	}
+
+	static float RunHPFConstantMono(FInterpolatedHPF& Filter, float InputValue, int32 NumSamples)
+	{
+		float Output = 0.0f;
+		for (int32 i = 0; i < NumSamples; ++i)
+		{
+			Filter.ProcessAudioFrame(&InputValue, &Output);
+		}
+		return Output;
+	}
+
+	// ================================================================== LPF ==
+
+	// LPF: DC passthrough — constant input should converge to itself.
+	IMPLEMENT_SIMPLE_AUTOMATION_TEST(
+		FInterpolatedLPFDCPassthrough,
+		"System.SignalProcessing.InterpolatedLPF.DCPassthrough",
+		EAutomationTestFlags_ApplicationContextMask | EAutomationTestFlags::EngineFilter)
+	bool FInterpolatedLPFDCPassthrough::RunTest(const FString&)
+	{
+		FInterpolatedLPF Filter;
+		Filter.Init(TestSampleRate, 1);
+		Filter.StartFrequencyInterpolation(TestCutoffHz);
+
+		// 500 samples is >> 5x the time constant at 1 kHz / 48 kHz
+		const float Last = RunLPFConstantMono(Filter, 1.0f, 500);
+		UTEST_TRUE("LPF output converges to 1.0 for DC input", FMath::IsNearlyEqual(Last, 1.0f, 1e-3f));
+		return true;
+	}
+
+	// LPF: Reset() must clear the delay memory so zero input immediately gives zero output.
+	IMPLEMENT_SIMPLE_AUTOMATION_TEST(
+		FInterpolatedLPFResetClearsMemory,
+		"System.SignalProcessing.InterpolatedLPF.ResetClearsMemory",
+		EAutomationTestFlags_ApplicationContextMask | EAutomationTestFlags::EngineFilter)
+	bool FInterpolatedLPFResetClearsMemory::RunTest(const FString&)
+	{
+		FInterpolatedLPF Filter;
+		Filter.Init(TestSampleRate, 1);
+		Filter.StartFrequencyInterpolation(TestCutoffHz);
+		RunLPFConstantMono(Filter, 1.0f, 100);
+
+		Filter.Reset();
+
+		float ZeroIn = 0.0f;
+		float Output = 0.0f;
+		Filter.ProcessAudioFrame(&ZeroIn, &Output);
+		UTEST_EQUAL("LPF output is 0 after Reset() + zero input", Output, 0.0f);
+		return true;
+	}
+
+	// LPF: ProcessAudioBuffer (single call) must produce the same samples as
+	//      looping ProcessAudioFrame one-by-one.
+	IMPLEMENT_SIMPLE_AUTOMATION_TEST(
+		FInterpolatedLPFBufferMatchesFrameLoop,
+		"System.SignalProcessing.InterpolatedLPF.ProcessBufferMatchesFrameLoop",
+		EAutomationTestFlags_ApplicationContextMask | EAutomationTestFlags::EngineFilter)
+	bool FInterpolatedLPFBufferMatchesFrameLoop::RunTest(const FString&)
+	{
+		constexpr int32 NumSamples = 128;
+
+		// Build a deterministic input with some variation (ramp + sign flip)
+		TArray<float> Input;
+		Input.SetNumUninitialized(NumSamples);
+		for (int32 i = 0; i < NumSamples; ++i)
+		{
+			Input[i] = ((i % 7) - 3) * 0.1f;
+		}
+
+		// Filter A: ProcessAudioBuffer
+		FInterpolatedLPF FilterA;
+		FilterA.Init(TestSampleRate, 1);
+		FilterA.StartFrequencyInterpolation(TestCutoffHz);
+		TArray<float> OutputA;
+		OutputA.SetNumZeroed(NumSamples);
+		FilterA.ProcessAudioBuffer(Input.GetData(), OutputA.GetData(), NumSamples);
+
+		// Filter B: ProcessAudioFrame loop
+		FInterpolatedLPF FilterB;
+		FilterB.Init(TestSampleRate, 1);
+		FilterB.StartFrequencyInterpolation(TestCutoffHz);
+		TArray<float> OutputB;
+		OutputB.SetNumZeroed(NumSamples);
+		for (int32 i = 0; i < NumSamples; ++i)
+		{
+			FilterB.ProcessAudioFrame(&Input[i], &OutputB[i]);
+		}
+
+		for (int32 i = 0; i < NumSamples; ++i)
+		{
+			if (!FMath::IsNearlyEqual(OutputA[i], OutputB[i], 1e-5f))
+			{
+				AddError(FString::Printf(TEXT("LPF output mismatch at sample %d: buffer=%.8f frame=%.8f"), i, OutputA[i], OutputB[i]));
+				return false;
+			}
+		}
+		return true;
+	}
+
+	// LPF: Two-channel filter must process channels independently.
+	//      Channel 0 = 1.0, Channel 1 = 0.0 → after settling, Ch0 ≈ 1, Ch1 ≈ 0.
+	IMPLEMENT_SIMPLE_AUTOMATION_TEST(
+		FInterpolatedLPFTwoChannelIndependence,
+		"System.SignalProcessing.InterpolatedLPF.TwoChannelIndependence",
+		EAutomationTestFlags_ApplicationContextMask | EAutomationTestFlags::EngineFilter)
+	bool FInterpolatedLPFTwoChannelIndependence::RunTest(const FString&)
+	{
+		FInterpolatedLPF Filter;
+		Filter.Init(TestSampleRate, 2);
+		Filter.StartFrequencyInterpolation(TestCutoffHz);
+
+		const float Frame[2] = { 1.0f, 0.0f };
+		float Out[2] = { 0.0f, 0.0f };
+		for (int32 i = 0; i < 500; ++i)
+		{
+			Filter.ProcessAudioFrame(Frame, Out);
+		}
+
+		UTEST_TRUE("LPF 2-ch: Ch0 converges to 1.0", FMath::IsNearlyEqual(Out[0], 1.0f, 1e-3f));
+		UTEST_TRUE("LPF 2-ch: Ch1 stays near 0.0",   FMath::IsNearlyEqual(Out[1], 0.0f, 1e-3f));
+		return true;
+	}
+
+	// LPF: first call to StartFrequencyInterpolation is applied instantly (no ramp).
+	IMPLEMENT_SIMPLE_AUTOMATION_TEST(
+		FInterpolatedLPFFirstFrequencyChangeIsInstant,
+		"System.SignalProcessing.InterpolatedLPF.FirstFrequencyChangeIsInstant",
+		EAutomationTestFlags_ApplicationContextMask | EAutomationTestFlags::EngineFilter)
+	bool FInterpolatedLPFFirstFrequencyChangeIsInstant::RunTest(const FString&)
+	{
+		// Two filters: one using InterpLength=1 on first change (should snap),
+		//              one using InterpLength=100 on first change (should also snap).
+		FInterpolatedLPF FilterSnap, FilterLong;
+		FilterSnap.Init(TestSampleRate, 1);
+		FilterLong.Init(TestSampleRate, 1);
+
+		FilterSnap.StartFrequencyInterpolation(TestCutoffHz, 1);
+		FilterLong.StartFrequencyInterpolation(TestCutoffHz, 100);
+
+		// After the first call both should have the same cutoff frequency
+		UTEST_EQUAL("LPF: first freq change is instant regardless of InterpLength",
+			FilterSnap.GetCutoffFrequency(), FilterLong.GetCutoffFrequency());
+
+		// And they should produce identical output (coefficient already at target)
+		float Input = 0.5f;
+		float OutSnap = 0.0f, OutLong = 0.0f;
+		FilterSnap.ProcessAudioFrame(&Input, &OutSnap);
+		FilterLong.ProcessAudioFrame(&Input, &OutLong);
+		UTEST_EQUAL("LPF: first-frame output identical regardless of initial InterpLength", OutSnap, OutLong);
+		return true;
+	}
+
+	// ================================================================== HPF ==
+
+	// HPF: DC rejection — constant input must converge to zero output.
+	IMPLEMENT_SIMPLE_AUTOMATION_TEST(
+		FInterpolatedHPFDCRejection,
+		"System.SignalProcessing.InterpolatedHPF.DCRejection",
+		EAutomationTestFlags_ApplicationContextMask | EAutomationTestFlags::EngineFilter)
+	bool FInterpolatedHPFDCRejection::RunTest(const FString&)
+	{
+		FInterpolatedHPF Filter;
+		Filter.Init(TestSampleRate, 1);
+		Filter.StartFrequencyInterpolation(TestCutoffHz);
+
+		const float Last = RunHPFConstantMono(Filter, 1.0f, 500);
+		UTEST_TRUE("HPF output converges to 0 for DC input", FMath::Abs(Last) < 1e-3f);
+		return true;
+	}
+
+	// HPF: Reset() must clear delay memory so zero input immediately gives zero output.
+	IMPLEMENT_SIMPLE_AUTOMATION_TEST(
+		FInterpolatedHPFResetClearsMemory,
+		"System.SignalProcessing.InterpolatedHPF.ResetClearsMemory",
+		EAutomationTestFlags_ApplicationContextMask | EAutomationTestFlags::EngineFilter)
+	bool FInterpolatedHPFResetClearsMemory::RunTest(const FString&)
+	{
+		FInterpolatedHPF Filter;
+		Filter.Init(TestSampleRate, 1);
+		Filter.StartFrequencyInterpolation(TestCutoffHz);
+		RunHPFConstantMono(Filter, 1.0f, 100);
+
+		Filter.Reset();
+
+		float ZeroIn = 0.0f;
+		float Output = 0.0f;
+		Filter.ProcessAudioFrame(&ZeroIn, &Output);
+		UTEST_EQUAL("HPF output is 0 after Reset() + zero input", Output, 0.0f);
+		return true;
+	}
+
+	// HPF: ProcessAudioBuffer must produce the same samples as looping ProcessAudioFrame.
+	IMPLEMENT_SIMPLE_AUTOMATION_TEST(
+		FInterpolatedHPFBufferMatchesFrameLoop,
+		"System.SignalProcessing.InterpolatedHPF.ProcessBufferMatchesFrameLoop",
+		EAutomationTestFlags_ApplicationContextMask | EAutomationTestFlags::EngineFilter)
+	bool FInterpolatedHPFBufferMatchesFrameLoop::RunTest(const FString&)
+	{
+		constexpr int32 NumSamples = 128;
+
+		TArray<float> Input;
+		Input.SetNumUninitialized(NumSamples);
+		for (int32 i = 0; i < NumSamples; ++i)
+		{
+			Input[i] = ((i % 7) - 3) * 0.1f;
+		}
+
+		FInterpolatedHPF FilterA;
+		FilterA.Init(TestSampleRate, 1);
+		FilterA.StartFrequencyInterpolation(TestCutoffHz);
+		TArray<float> OutputA;
+		OutputA.SetNumZeroed(NumSamples);
+		FilterA.ProcessAudioBuffer(Input.GetData(), OutputA.GetData(), NumSamples);
+
+		FInterpolatedHPF FilterB;
+		FilterB.Init(TestSampleRate, 1);
+		FilterB.StartFrequencyInterpolation(TestCutoffHz);
+		TArray<float> OutputB;
+		OutputB.SetNumZeroed(NumSamples);
+		for (int32 i = 0; i < NumSamples; ++i)
+		{
+			FilterB.ProcessAudioFrame(&Input[i], &OutputB[i]);
+		}
+
+		for (int32 i = 0; i < NumSamples; ++i)
+		{
+			if (!FMath::IsNearlyEqual(OutputA[i], OutputB[i], 1e-5f))
+			{
+				AddError(FString::Printf(TEXT("HPF output mismatch at sample %d: buffer=%.8f frame=%.8f"), i, OutputA[i], OutputB[i]));
+				return false;
+			}
+		}
+		return true;
+	}
+
+	// HPF: Two-channel filter must process channels independently.
+	//      Both channels fed DC 1.0, both should converge to 0.
+	IMPLEMENT_SIMPLE_AUTOMATION_TEST(
+		FInterpolatedHPFTwoChannelIndependence,
+		"System.SignalProcessing.InterpolatedHPF.TwoChannelIndependence",
+		EAutomationTestFlags_ApplicationContextMask | EAutomationTestFlags::EngineFilter)
+	bool FInterpolatedHPFTwoChannelIndependence::RunTest(const FString&)
+	{
+		FInterpolatedHPF Filter;
+		Filter.Init(TestSampleRate, 2);
+		Filter.StartFrequencyInterpolation(TestCutoffHz);
+
+		// Ch0 = DC 1.0, Ch1 = DC -0.5 — both should reject DC
+		const float Frame[2] = { 1.0f, -0.5f };
+		float Out[2] = { 0.0f, 0.0f };
+		for (int32 i = 0; i < 500; ++i)
+		{
+			Filter.ProcessAudioFrame(Frame, Out);
+		}
+
+		UTEST_TRUE("HPF 2-ch: Ch0 converges to 0",  FMath::Abs(Out[0]) < 1e-3f);
+		UTEST_TRUE("HPF 2-ch: Ch1 converges to 0",  FMath::Abs(Out[1]) < 1e-3f);
+		return true;
+	}
+
+	// HPF + LPF: processing the same signal through both filters should confirm
+	//            that the HPF output is always (Input - internal_LPF_component),
+	//            i.e. HPF output + LPF-component = Input.
+	//            We verify this by checking HPF + separate FInterpolatedLPF ≠ Input
+	//            (they use different designs) while confirming HPF alone converges correctly.
+	//            The practical invariant we test: HPF(DC) → 0 and LPF(DC) → 1 independently.
+	IMPLEMENT_SIMPLE_AUTOMATION_TEST(
+		FInterpolatedLPFAndHPFComplementary,
+		"System.SignalProcessing.InterpolatedOnePole.LPFAndHPFComplementary",
+		EAutomationTestFlags_ApplicationContextMask | EAutomationTestFlags::EngineFilter)
+	bool FInterpolatedLPFAndHPFComplementary::RunTest(const FString&)
+	{
+		FInterpolatedLPF LPF;
+		FInterpolatedHPF HPF;
+		LPF.Init(TestSampleRate, 1);
+		HPF.Init(TestSampleRate, 1);
+		LPF.StartFrequencyInterpolation(TestCutoffHz);
+		HPF.StartFrequencyInterpolation(TestCutoffHz);
+
+		float InputSample = 1.0f;
+		float LPFOut = 0.0f, HPFOut = 0.0f;
+		for (int32 i = 0; i < 500; ++i)
+		{
+			LPF.ProcessAudioFrame(&InputSample, &LPFOut);
+			HPF.ProcessAudioFrame(&InputSample, &HPFOut);
+		}
+
+		// After settling, LPF should pass DC and HPF should block it
+		UTEST_TRUE("LPF converges to 1.0 for DC", FMath::IsNearlyEqual(LPFOut, 1.0f, 1e-3f));
+		UTEST_TRUE("HPF converges to 0.0 for DC",  FMath::Abs(HPFOut) < 1e-3f);
+		return true;
+	}
+
+} // namespace Audio::Tests
+
+#endif // WITH_DEV_AUTOMATION_TESTS
